- Move AnimBlueprintLibrary to its own Module (out of AnimModifiers)
- Added FAnimPose, a script-friendly helper structure to generate/read/write evaluated anim poses
- Added API to populate AnimPose from AnimationSequenceBase, USkeleton (ref pose) and AnimBlueprint

diff --git a/Engine/Config/BaseEngine.ini b/Engine/Config/BaseEngine.ini
--- a/Engine/Config/BaseEngine.ini
+++ b/Engine/Config/BaseEngine.ini
@@ -926,6 +926,7 @@ NearClipPlane=10.0
 +ClassRedirects=(OldName="/Script/PropertyAccess.AnimBlueprintClassSubsystem_PropertyAccess",NewName="/Script/Engine.AnimBlueprintClassSubsystem_PropertyAccess")
 +PackageRedirects=(OldName="/Script/PropertyAccess", Removed=true)
 +StructRedirects=(OldName="AnimInstanceSubsystemData",NewName="/Script/Engine.AnimSubsystemInstance")
++ClassRedirects=(OldName="/Script/AnimationModifiers.AnimationBlueprintLibrary",NewName="/Script/AnimationBlueprintLibrary.AnimationBlueprintLibrary")
 
 [CoreUObject.Metadata]
 ; Note: UnrealHeaderTool should be rerun after making changes to MetadataRedirects to catch any keys specified in class headers
diff --git a/Engine/Plugins/Experimental/SkeletalReduction/Source/SkeletalMeshReduction.Build.cs b/Engine/Plugins/Experimental/SkeletalReduction/Source/SkeletalMeshReduction.Build.cs
--- a/Engine/Plugins/Experimental/SkeletalReduction/Source/SkeletalMeshReduction.Build.cs
+++ b/Engine/Plugins/Experimental/SkeletalReduction/Source/SkeletalMeshReduction.Build.cs
@@ -55,7 +55,7 @@ namespace UnrealBuildTool.Rules
                     "Engine",
 					"EditorFramework",
                     "UnrealEd",
-                    "AnimationModifiers",
+					"AnimationBlueprintLibrary",
                     "MeshBoneReduction",
                     "QuadricMeshReduction",
                     "ClothingSystemRuntimeCommon"
diff --git a/Engine/Source/Developer/MeshBoneReduction/MeshBoneReduction.Build.cs b/Engine/Source/Developer/MeshBoneReduction/MeshBoneReduction.Build.cs
--- a/Engine/Source/Developer/MeshBoneReduction/MeshBoneReduction.Build.cs
+++ b/Engine/Source/Developer/MeshBoneReduction/MeshBoneReduction.Build.cs
@@ -13,6 +13,6 @@ public class MeshBoneReduction : ModuleRules
 		PrivateDependencyModuleNames.Add("Engine");
         PrivateDependencyModuleNames.Add("RenderCore");
 		PrivateDependencyModuleNames.Add("RHI");
-        PrivateDependencyModuleNames.Add("AnimationModifiers");
+        PrivateDependencyModuleNames.Add("AnimationBlueprintLibrary");
     }
 }
diff --git a/Engine/Source/Editor/AnimationBlueprintLibrary/AnimationBlueprintLibrary.build.cs b/Engine/Source/Editor/AnimationBlueprintLibrary/AnimationBlueprintLibrary.build.cs
new file mode 100644
--- /dev/null
+++ b/Engine/Source/Editor/AnimationBlueprintLibrary/AnimationBlueprintLibrary.build.cs
@@ -0,0 +1,19 @@
+// Copyright Epic Games, Inc. All Rights Reserved.
+
+using UnrealBuildTool;
+
+public class AnimationBlueprintLibrary : ModuleRules
+{
+	public AnimationBlueprintLibrary(ReadOnlyTargetRules Target) : base(Target)
+	{
+		PrivateDependencyModuleNames.AddRange(
+			new string[] 
+            {
+                "Core",
+                "CoreUObject",
+                "Engine",
+                "Kismet"              
+            }
+		);
+    }
+}
diff --git a/Engine/Source/Editor/AnimationBlueprintLibrary/Private/AnimPose.cpp b/Engine/Source/Editor/AnimationBlueprintLibrary/Private/AnimPose.cpp
new file mode 100644
--- /dev/null
+++ b/Engine/Source/Editor/AnimationBlueprintLibrary/Private/AnimPose.cpp
@@ -0,0 +1,496 @@
+// Copyright Epic Games, Inc. All Rights Reserved.
+
+#include "AnimPose.h"
+#include "PreviewScene.h"
+#include "Animation/AnimNode_LinkedInputPose.h"
+
+DEFINE_LOG_CATEGORY_STATIC(LogAnimationPoseScripting, Verbose, All);
+
+void FAnimPose::Init(const FBoneContainer& InBoneContainer)
+{
+	Reset();
+	BoneContainer = InBoneContainer;
+		
+	for (const FBoneIndexType BoneIndex : BoneContainer.GetBoneIndicesArray())
+	{			
+		const int32 SkeletonBoneIndex = BoneContainer.GetSkeletonIndex(FCompactPoseBoneIndex(BoneIndex));
+		const int32 ParentBoneIndex = BoneContainer.GetParentBoneIndex(SkeletonBoneIndex);
+
+		BoneIndices.Add(SkeletonBoneIndex);
+		ParentBoneIndices.Add(ParentBoneIndex);
+
+		const FReferenceSkeleton& RefSkeleton = BoneContainer.GetReferenceSkeleton();
+		BoneNames.Add(RefSkeleton.GetBoneName(SkeletonBoneIndex));
+
+		RefLocalSpacePoses.Add(BoneContainer.GetRefPoseTransform(FCompactPoseBoneIndex(BoneIndex)));
+	}
+
+	TArray<bool> Processed;
+	Processed.SetNumZeroed(BoneNames.Num());
+	RefWorldSpacePoses.SetNum(BoneNames.Num());
+	for (int32 EntryIndex = 0; EntryIndex < BoneNames.Num(); ++EntryIndex)
+	{
+		const int32 ParentIndex = ParentBoneIndices[EntryIndex];
+		const int32 TransformIndex = BoneIndices.IndexOfByKey(ParentIndex);
+
+		if (TransformIndex != INDEX_NONE)
+		{
+			ensure(Processed[TransformIndex]);
+			RefWorldSpacePoses[EntryIndex] = RefLocalSpacePoses[EntryIndex] * RefWorldSpacePoses[TransformIndex];
+		}
+		else
+		{
+			RefWorldSpacePoses[EntryIndex] = RefLocalSpacePoses[EntryIndex];
+		}
+
+		Processed[EntryIndex] = true;
+	}
+}
+
+void FAnimPose::GetPose(FCompactPose& InOutCompactPose) const
+{
+	if (IsValid())
+	{
+		for (int32 Index = 0; Index < BoneNames.Num(); ++Index)
+		{
+			const FName& BoneName = BoneNames[Index];
+			const FCompactPoseBoneIndex PoseBoneIndex = FCompactPoseBoneIndex(InOutCompactPose.GetBoneContainer().GetPoseBoneIndexForBoneName(BoneName));
+			if (PoseBoneIndex != INDEX_NONE)
+			{
+				InOutCompactPose[PoseBoneIndex] = LocalSpacePoses[Index];
+			}
+		}
+	}
+}
+
+void FAnimPose::SetPose(USkeletalMeshComponent* Component)
+{
+	if (IsInitialized())
+	{		
+		const FBoneContainer& ContextBoneContainer = Component->GetAnimInstance()->GetRequiredBones();
+
+		LocalSpacePoses.SetNum(RefLocalSpacePoses.Num());
+
+		TArray<FTransform> BoneSpaceTransforms = Component->GetBoneSpaceTransforms();
+		for (const FBoneIndexType BoneIndex : ContextBoneContainer.GetBoneIndicesArray())
+		{
+			const int32 SkeletonBoneIndex = ContextBoneContainer.GetSkeletonIndex(FCompactPoseBoneIndex(BoneIndex));
+			LocalSpacePoses[BoneIndices.IndexOfByKey(SkeletonBoneIndex)] = BoneSpaceTransforms[BoneIndex];
+		}
+
+		ensure(LocalSpacePoses.Num() == RefLocalSpacePoses.Num());	
+		GenerateWorldSpaceTransforms();
+	}
+}
+
+void FAnimPose::GenerateWorldSpaceTransforms()
+{
+	if (IsPopulated())
+	{
+		TArray<bool> Processed;
+		Processed.SetNumZeroed(BoneNames.Num());
+		WorldSpacePoses.SetNum(BoneNames.Num());
+		for (int32 EntryIndex = 0; EntryIndex < BoneNames.Num(); ++EntryIndex)
+		{
+			const int32 ParentIndex = ParentBoneIndices[EntryIndex];
+			const int32 TransformIndex = BoneIndices.IndexOfByKey(ParentIndex);
+			if (TransformIndex != INDEX_NONE)
+			{
+				ensure(Processed[TransformIndex]);
+				WorldSpacePoses[EntryIndex] = LocalSpacePoses[EntryIndex] * WorldSpacePoses[TransformIndex];
+			}
+			else
+			{
+				WorldSpacePoses[EntryIndex] = LocalSpacePoses[EntryIndex];
+			}
+
+			Processed[EntryIndex] = true;
+		}
+	}
+	else
+	{
+		UE_LOG(LogAnimationPoseScripting, Warning, TEXT("Anim Pose was not previously populated"));
+	}
+}
+
+void FAnimPose::SetPose(const FCompactPose& CompactPose)
+{
+	if (IsInitialized())
+	{
+		const FBoneContainer& ContextBoneContainer = CompactPose.GetBoneContainer();
+			
+		LocalSpacePoses.SetNum(RefLocalSpacePoses.Num());
+		for (const FCompactPoseBoneIndex BoneIndex : CompactPose.ForEachBoneIndex())
+		{
+			const int32 SkeletonBoneIndex = ContextBoneContainer.GetSkeletonIndex(BoneIndex);
+			LocalSpacePoses[BoneIndices.IndexOfByKey(SkeletonBoneIndex)] = CompactPose[BoneIndex];
+		}
+
+		ensure(LocalSpacePoses.Num() == RefLocalSpacePoses.Num());
+		GenerateWorldSpaceTransforms();
+	}
+	else
+	{
+		UE_LOG(LogAnimationPoseScripting, Warning, TEXT("Anim Pose was not previously initialized"));
+	}
+}
+
+void FAnimPose::SetToRefPose()
+{
+	if (IsInitialized())
+	{
+		LocalSpacePoses = RefLocalSpacePoses;
+		WorldSpacePoses = RefWorldSpacePoses;
+	}
+	else
+	{
+		UE_LOG(LogAnimationPoseScripting, Warning, TEXT("Anim Pose was not previously initialized"));
+	}
+}
+
+bool FAnimPose::IsValid() const
+{
+	const int32 ExpectedNumBones = BoneNames.Num();
+	bool bIsValid = ExpectedNumBones != 0;
+	
+	bIsValid &= BoneIndices.Num() == ExpectedNumBones;
+	bIsValid &= ParentBoneIndices.Num() == ExpectedNumBones;
+	bIsValid &= LocalSpacePoses.Num() == ExpectedNumBones;
+	bIsValid &= WorldSpacePoses.Num() == ExpectedNumBones;
+	bIsValid &= RefLocalSpacePoses.Num() == ExpectedNumBones;
+	bIsValid &= RefWorldSpacePoses.Num() == ExpectedNumBones;
+	bIsValid &= BoneContainer.GetNumBones() == ExpectedNumBones;
+	
+	return bIsValid;
+}
+
+void FAnimPose::Reset()
+{
+	BoneNames.Empty();
+	BoneIndices.Empty();
+	ParentBoneIndices.Empty();
+	LocalSpacePoses.Empty();
+	WorldSpacePoses.Empty();
+	RefLocalSpacePoses.Empty();
+	RefWorldSpacePoses.Empty();
+}
+
+bool UAnimPoseExtensions::IsValid(const FAnimPose& Pose)
+{
+	return Pose.IsValid();
+}
+
+void UAnimPoseExtensions::GetBoneNames(const FAnimPose& Pose, TArray<FName>& Bones)
+{
+	Bones.Append(Pose.BoneNames);
+}
+
+FTransform UAnimPoseExtensions::GetBonePose(const FAnimPose& Pose, FName BoneName, EAnimPoseSpaces Space /*= EAnimPoseSpaces::Local*/)
+{
+	if (Pose.IsValid())
+	{
+		const int32 BoneIndex = Pose.BoneNames.IndexOfByKey(BoneName);
+		if (BoneIndex != INDEX_NONE)
+		{		
+			return Space == EAnimPoseSpaces::Local ? Pose.LocalSpacePoses[BoneIndex] : Pose.WorldSpacePoses[BoneIndex];		
+		}
+		else
+		{
+			UE_LOG(LogAnimationPoseScripting, Warning, TEXT("No bone with name %s was found"), *BoneName.ToString());
+		}
+	}
+	else
+	{		
+		UE_LOG(LogAnimationPoseScripting, Error, TEXT("Provided Pose is not valid"));	
+	}
+	return FTransform::Identity;
+}
+
+void UAnimPoseExtensions::SetBonePose(FAnimPose& Pose, FTransform Transform, FName BoneName, EAnimPoseSpaces Space /*= EAnimPoseSpaces::Local*/)
+{
+	if (Pose.IsValid())
+	{
+		const int32 BoneIndex = Pose.BoneNames.IndexOfByKey(BoneName);
+		if (BoneIndex != INDEX_NONE)
+		{
+			if (Space == EAnimPoseSpaces::Local)
+			{
+				Pose.LocalSpacePoses[BoneIndex] = Transform;
+			}
+			else if (Space == EAnimPoseSpaces::World)
+			{
+				const int32 ParentIndex = Pose.ParentBoneIndices[BoneIndex];
+				const FTransform ParentTransformWS = ParentIndex != INDEX_NONE ? Pose.WorldSpacePoses[ParentIndex] : FTransform::Identity;
+				Pose.LocalSpacePoses[BoneIndex] = Transform.GetRelativeTransform(ParentTransformWS);
+			}
+			
+			Pose.GenerateWorldSpaceTransforms();
+		}
+		else
+		{
+			UE_LOG(LogAnimationPoseScripting, Warning, TEXT("No bone with name %s was found"), *BoneName.ToString());
+		}
+	}
+	else
+	{		
+		UE_LOG(LogAnimationPoseScripting, Error, TEXT("Provided Pose is not valid"));	
+	}
+}
+	
+FTransform UAnimPoseExtensions::GetRefBonePose(const FAnimPose& Pose, FName BoneName, EAnimPoseSpaces Space /*= EAnimPoseSpaces::Local*/)
+{
+	if (Pose.IsValid())
+	{
+		const int32 BoneIndex = Pose.BoneNames.IndexOfByKey(BoneName);
+		if (BoneIndex != INDEX_NONE)
+		{	
+			return Space == EAnimPoseSpaces::Local ? Pose.RefLocalSpacePoses[BoneIndex] : Pose.RefWorldSpacePoses[BoneIndex];
+		}
+		else
+		{
+			UE_LOG(LogAnimationPoseScripting, Warning, TEXT("No bone with name %s was found"), *BoneName.ToString());
+		}
+	}
+	else
+	{		
+		UE_LOG(LogAnimationPoseScripting, Error, TEXT("Provided Pose is not valid"));	
+	}
+
+	return FTransform::Identity;	
+}
+
+FTransform UAnimPoseExtensions::GetRelativeTransform(const FAnimPose& Pose, FName FromBoneName, FName ToBoneName, EAnimPoseSpaces Space /*= EAnimPoseSpaces::Local*/)
+{
+	if (Pose.IsValid())
+	{
+		const int32 FromBoneIndex = Pose.BoneNames.IndexOfByKey(FromBoneName);
+		const int32 ToBoneIndex = Pose.BoneNames.IndexOfByKey(ToBoneName);
+		if (FromBoneIndex != INDEX_NONE && ToBoneIndex != INDEX_NONE)
+		{	
+			const FTransform& From = Space == EAnimPoseSpaces::Local ? Pose.LocalSpacePoses[FromBoneIndex] : Pose.WorldSpacePoses[FromBoneIndex];
+			const FTransform& To = Space == EAnimPoseSpaces::Local ? Pose.LocalSpacePoses[ToBoneIndex] : Pose.WorldSpacePoses[ToBoneIndex];
+
+			const FTransform Relative = To.GetRelativeTransform(From);
+				
+			return Relative;
+		}
+		else
+		{
+			UE_LOG(LogAnimationPoseScripting, Warning, TEXT("No bone with name %s or %s was found"), *FromBoneName.ToString(), *ToBoneName.ToString());
+		}
+	}
+	else
+	{		
+		UE_LOG(LogAnimationPoseScripting, Error, TEXT("Provided Pose is not valid"));	
+	}
+
+	return FTransform::Identity;
+}
+
+FTransform UAnimPoseExtensions::GetRelativeToRefPoseTransform(const FAnimPose& Pose, FName BoneName, EAnimPoseSpaces Space /*= EAnimPoseSpaces::Local*/)
+{
+	if (Pose.IsValid())
+	{
+		const int32 BoneIndex = Pose.BoneNames.IndexOfByKey(BoneName);
+		if (BoneIndex != INDEX_NONE)
+		{	
+			const FTransform& From = Space == EAnimPoseSpaces::Local ? Pose.RefLocalSpacePoses[BoneIndex] : Pose.RefWorldSpacePoses[BoneIndex];
+			const FTransform& To = Space == EAnimPoseSpaces::Local ? Pose.LocalSpacePoses[BoneIndex] : Pose.WorldSpacePoses[BoneIndex];
+
+			return To.GetRelativeTransform(From);
+		}
+		else
+		{
+			UE_LOG(LogAnimationPoseScripting, Warning, TEXT("No bone with name %s was found"), *BoneName.ToString());
+		}
+	}
+	else
+	{		
+		UE_LOG(LogAnimationPoseScripting, Error, TEXT("Provided Pose is not valid"));	
+	}
+
+	return FTransform::Identity;
+}
+
+FTransform UAnimPoseExtensions::GetRefPoseRelativeTransform(const FAnimPose& Pose, FName FromBoneName, FName ToBoneName, EAnimPoseSpaces Space /*= EAnimPoseSpaces::Local*/)
+{
+	if (Pose.IsValid())
+	{
+		const int32 FromBoneIndex = Pose.BoneNames.IndexOfByKey(FromBoneName);
+		const int32 ToBoneIndex = Pose.BoneNames.IndexOfByKey(ToBoneName);
+		if (FromBoneIndex != INDEX_NONE && ToBoneIndex != INDEX_NONE)
+		{	
+			const FTransform& From = Space == EAnimPoseSpaces::Local ? Pose.RefLocalSpacePoses[FromBoneIndex] : Pose.RefWorldSpacePoses[FromBoneIndex];
+			const FTransform& To = Space == EAnimPoseSpaces::Local ? Pose.RefLocalSpacePoses[ToBoneIndex] : Pose.RefWorldSpacePoses[ToBoneIndex];
+
+			const FTransform Relative = From.GetRelativeTransform(To);
+				
+			return Relative;
+		}
+		else
+		{
+			UE_LOG(LogAnimationPoseScripting, Warning, TEXT("No bone with name %s or %s was found"), *FromBoneName.ToString(), *ToBoneName.ToString());
+		}
+	}
+	else
+	{		
+		UE_LOG(LogAnimationPoseScripting, Error, TEXT("Provided Pose is not valid"));	
+	}
+
+	return FTransform::Identity;
+}
+
+void UAnimPoseExtensions::EvaluateAnimationBlueprintWithInputPose(const FAnimPose& Pose, USkeletalMesh* TargetSkeletalMesh, UAnimBlueprint* AnimationBlueprint, FAnimPose& OutPose)
+{
+	if (Pose.IsValid())
+	{
+		if (TargetSkeletalMesh)
+		{
+			if (AnimationBlueprint)
+			{
+				UAnimBlueprintGeneratedClass* AnimGeneratedClass = AnimationBlueprint->GetAnimBlueprintGeneratedClass();
+				if (AnimGeneratedClass)
+				{
+					if (AnimGeneratedClass->TargetSkeleton == TargetSkeletalMesh->GetSkeleton())
+					{
+						FMemMark Mark(FMemStack::Get());
+						
+						FPreviewScene PreviewScene;
+			
+						USkeletalMeshComponent* Component = NewObject<USkeletalMeshComponent>();
+						Component->SetSkeletalMesh(TargetSkeletalMesh);
+						Component->SetAnimInstanceClass(AnimationBlueprint->GetAnimBlueprintGeneratedClass());
+
+						PreviewScene.AddComponent(Component, FTransform::Identity);
+			
+						if (UAnimInstance* AnimInstance = Component->GetAnimInstance())
+						{
+							if (FAnimNode_LinkedInputPose* InputNode = AnimInstance->GetLinkedInputPoseNode())
+							{
+								InputNode->CachedInputPose.SetBoneContainer(&AnimInstance->GetRequiredBones());
+								InputNode->CachedInputCurve.InitFrom(AnimInstance->GetRequiredBones());
+								InputNode->CachedInputPose.CopyBonesFrom(Pose.LocalSpacePoses);
+					
+								OutPose.Init(AnimInstance->GetRequiredBones());
+
+								Component->InitAnim(true);
+								Component->RefreshBoneTransforms();
+								const TArray<FTransform>& LocalSpaceTransforms = Component->GetBoneSpaceTransforms();
+
+								OutPose.SetPose(Component);	
+							}	
+						}
+						else
+						{
+							UE_LOG(LogAnimationPoseScripting, Warning, TEXT("Failed to retrieve Input Pose Node from Animation Graph %s"), *AnimationBlueprint->GetName());	
+						}
+					}
+					else
+					{
+						UE_LOG(LogAnimationPoseScripting, Error, TEXT("Animation Blueprint target Skeleton %s does not match Target Skeletal Mesh its Skeleton %s"), *AnimGeneratedClass->TargetSkeleton->GetName(), *TargetSkeletalMesh->GetSkeleton()->GetName());	
+					}
+				
+				}
+				else
+				{
+					UE_LOG(LogAnimationPoseScripting, Warning, TEXT("Failed to retrieve Animation Blueprint generated class"));	
+				}
+			}
+			else
+			{
+				UE_LOG(LogAnimationPoseScripting, Warning, TEXT("Invalid Animation Blueprint"));			
+			}	
+		}
+		else
+		{		
+			UE_LOG(LogAnimationPoseScripting, Warning, TEXT("Invalid Target Skeletal Mesh"));
+		}
+	}
+	else
+	{		
+		UE_LOG(LogAnimationPoseScripting, Error, TEXT("Provided Pose is not valid"));	
+	}
+}
+
+void UAnimPoseExtensions::GetReferencePose(USkeleton* Skeleton, FAnimPose& OutPose)
+{
+	if (Skeleton)
+	{
+		const FReferenceSkeleton& RefSkeleton = Skeleton->GetReferenceSkeleton();
+				       
+		TArray<FBoneIndexType> RequiredBoneIndexArray;
+		RequiredBoneIndexArray.AddUninitialized(RefSkeleton.GetNum());
+		for (int32 BoneIndex = 0; BoneIndex < RequiredBoneIndexArray.Num(); ++BoneIndex)
+		{
+			RequiredBoneIndexArray[BoneIndex] = BoneIndex;
+		}
+
+		FBoneContainer RequiredBones;
+		RequiredBones.InitializeTo(RequiredBoneIndexArray, FCurveEvaluationOption(false), *Skeleton);
+
+		OutPose.Init(RequiredBones);
+		OutPose.SetToRefPose();
+	}
+	else
+	{		
+		UE_LOG(LogAnimationPoseScripting, Error, TEXT("Invalid Skeleton provided"));	
+	}
+}
+
+void UAnimPoseExtensions::GetAnimPoseAtFrame(const UAnimSequenceBase* AnimationSequenceBase, int32 FrameIndex, FAnimPoseEvaluationOptions EvaluationOptions, FAnimPose& Pose)
+{
+	float Time = 0.f;
+	UAnimationBlueprintLibrary::GetTimeAtFrame(AnimationSequenceBase, FrameIndex, Time);
+	GetAnimPoseAtTime(AnimationSequenceBase, Time, EvaluationOptions, Pose);
+}
+
+void UAnimPoseExtensions::GetAnimPoseAtTime(const UAnimSequenceBase* AnimationSequenceBase, float Time, FAnimPoseEvaluationOptions EvaluationOptions, FAnimPose& Pose)
+{
+	if (AnimationSequenceBase && AnimationSequenceBase->GetSkeleton())
+	{
+		FMemMark Mark(FMemStack::Get());
+		const FReferenceSkeleton& RefSkeleton = AnimationSequenceBase->GetSkeleton()->GetReferenceSkeleton();
+
+		bool bValidTime = false;
+		UAnimationBlueprintLibrary::IsValidTime(AnimationSequenceBase, Time, bValidTime);
+		if (bValidTime)
+		{			       
+			TArray<FBoneIndexType> RequiredBoneIndexArray;
+			RequiredBoneIndexArray.AddUninitialized(RefSkeleton.GetNum());
+			for (int32 BoneIndex = 0; BoneIndex < RequiredBoneIndexArray.Num(); ++BoneIndex)
+			{
+				RequiredBoneIndexArray[BoneIndex] = BoneIndex;
+			}
+
+			FBoneContainer RequiredBones;
+			RequiredBones.InitializeTo(RequiredBoneIndexArray, FCurveEvaluationOption(false), *AnimationSequenceBase->GetSkeleton());
+			
+			RequiredBones.SetUseRAWData(EvaluationOptions.EvaluationType == EAnimDataEvalType::Raw);
+			RequiredBones.SetUseSourceData(EvaluationOptions.EvaluationType == EAnimDataEvalType::Source);
+			
+			RequiredBones.SetDisableRetargeting(!EvaluationOptions.bShouldRetarget);
+			
+			FCompactPose CompactPose;
+			CompactPose.SetBoneContainer(&RequiredBones);
+
+			Pose.Init(CompactPose.GetBoneContainer());
+
+			FBlendedCurve Curve;
+			UE::Anim::FStackAttributeContainer Attributes;
+
+			FAnimationPoseData PoseData(CompactPose, Curve, Attributes);
+			FAnimExtractContext Context(Time, EvaluationOptions.bExtractRootMotion);
+
+			AnimationSequenceBase->GetAnimationPose(PoseData, Context);
+			Pose.SetPose(CompactPose);
+		}
+		else
+		{
+			UE_LOG(LogAnimationPoseScripting, Warning, TEXT("Invalid time value %f for Animation Sequence %s supplied for GetBonePosesForTime"), Time, *AnimationSequenceBase->GetName());
+		}
+	}
+	else
+	{
+		UE_LOG(LogAnimationPoseScripting, Warning, TEXT("Invalid Animation Sequence supplied for GetBonePosesForTime"));
+	}
+}
\ No newline at end of file
diff --git a/Engine/Source/Editor/AnimationBlueprintLibrary/Private/AnimationBlueprintLibrary.cpp b/Engine/Source/Editor/AnimationBlueprintLibrary/Private/AnimationBlueprintLibrary.cpp
new file mode 100644
--- /dev/null
+++ b/Engine/Source/Editor/AnimationBlueprintLibrary/Private/AnimationBlueprintLibrary.cpp
@@ -0,0 +1,2270 @@
+// Copyright Epic Games, Inc. All Rights Reserved.
+
+#include "AnimationBlueprintLibrary.h" 
+
+#include "Animation/AnimSequence.h"
+#include "Animation/AnimationAsset.h"
+#include "Animation/AnimBoneCompressionSettings.h"
+#include "Animation/AnimCurveCompressionSettings.h"
+#include "Animation/AnimMetaData.h"
+#include "Animation/AnimNotifies/AnimNotifyState.h"
+#include "Animation/Skeleton.h"
+#include "Animation/AnimNotifies/AnimNotify.h"
+#include "Animation/AnimCurveTypes.h"
+#include "BonePose.h" 
+#include "Algo/Transform.h"
+
+#include "AnimationRuntime.h"
+#include "Animation/AnimData/AnimDataModel.h"
+#include "Animation/AnimSequenceHelpers.h"
+
+#define LOCTEXT_NAMESPACE "AnimationBlueprintLibrary"
+
+
+IMPLEMENT_MODULE(IModuleInterface, AnimationBlueprintLibrary);
+
+DEFINE_LOG_CATEGORY_STATIC(LogAnimationBlueprintLibrary, Verbose, All);
+
+void UAnimationBlueprintLibrary::GetNumFrames(const UAnimSequenceBase* AnimationSequenceBase, int32& NumFrames)
+{
+	NumFrames = 0;
+	if (AnimationSequenceBase)
+	{
+		NumFrames = AnimationSequenceBase->GetNumberOfSampledKeys() - 1;
+	}
+	else
+	{
+		UE_LOG(LogAnimationBlueprintLibrary, Warning, TEXT("Invalid Animation Sequence supplied for GetNumFrames"));
+	}
+}
+
+void UAnimationBlueprintLibrary::GetNumKeys(const UAnimSequenceBase* AnimationSequenceBase, int32& NumKeys)
+{
+	NumKeys = 0;
+	if (AnimationSequenceBase)
+	{
+		NumKeys = AnimationSequenceBase->GetNumberOfSampledKeys();
+	}
+	else
+	{
+		UE_LOG(LogAnimationBlueprintLibrary, Warning, TEXT("Invalid Animation Sequence supplied for GetNumKeys"));
+	}
+}
+
+void UAnimationBlueprintLibrary::GetAnimationTrackNames(const UAnimSequenceBase* AnimationSequenceBase, TArray<FName>& TrackNames)
+{
+	TrackNames.Empty();
+	if (AnimationSequenceBase)
+	{
+		AnimationSequenceBase->GetDataModel()->GetBoneTrackNames(TrackNames);
+	}
+	else
+	{
+		UE_LOG(LogAnimationBlueprintLibrary, Warning, TEXT("Invalid Animation Sequence supplied for GetBoneTrackNames"));
+	}	
+}
+
+void UAnimationBlueprintLibrary::GetAnimationCurveNames(const UAnimSequence* AnimationSequence, ERawCurveTrackTypes CurveType, TArray<FName>& CurveNames)
+{
+	CurveNames.Empty();
+	if (AnimationSequence)
+	{
+		auto GetCurveName = [](const auto& Curve) -> FName
+		{
+			return Curve.Name.DisplayName;
+		};
+
+		switch (CurveType)
+		{
+			case ERawCurveTrackTypes::RCT_Float:
+			{
+				Algo::Transform(AnimationSequence->GetDataModel()->GetFloatCurves(), CurveNames, GetCurveName);
+				break;
+			}
+
+			case ERawCurveTrackTypes::RCT_Transform:
+			{
+				Algo::Transform(AnimationSequence->GetDataModel()->GetTransformCurves(), CurveNames, GetCurveName);
+				break;
+			}
+
+			default:
+			{
+				UE_LOG(LogAnimationBlueprintLibrary, Warning, TEXT("Invalid CurveType supplied for GetCurveNames"));
+			}
+		}
+		
+	}
+	else
+	{
+		UE_LOG(LogAnimationBlueprintLibrary, Warning, TEXT("Invalid Animation Sequence supplied for GetCurveNames"));
+	}
+}
+
+void UAnimationBlueprintLibrary::GetRawTrackPositionData(const UAnimSequenceBase* AnimationSequenceBase, const FName TrackName, TArray<FVector>& PositionData)
+{
+	PositionData.Empty();
+	if (IsValidRawAnimationTrackName(AnimationSequenceBase, TrackName))
+	{
+		const FRawAnimSequenceTrack& RawTrack =	GetRawAnimationTrackByName(AnimationSequenceBase, TrackName);
+		PositionData.Append(RawTrack.PosKeys);
+	}
+	else
+	{
+		const FString AnimSequenceName = (AnimationSequenceBase != nullptr) ? AnimationSequenceBase->GetName() : "Invalid Animation sequence";
+		UE_LOG(LogAnimationBlueprintLibrary, Warning, TEXT("Raw Animation Track name %s does not exist in Animation Sequence %s"), *TrackName.ToString(), *AnimSequenceName );
+	}	
+}
+
+void UAnimationBlueprintLibrary::GetRawTrackRotationData(const UAnimSequenceBase* AnimationSequenceBase, const FName TrackName, TArray<FQuat>& RotationData)
+{
+	RotationData.Empty();
+	if (IsValidRawAnimationTrackName(AnimationSequenceBase, TrackName))
+	{
+		const FRawAnimSequenceTrack& RawTrack = GetRawAnimationTrackByName(AnimationSequenceBase, TrackName);
+		RotationData.Append(RawTrack.RotKeys);
+	}
+	else
+	{	
+		const FString AnimSequenceName = (AnimationSequenceBase != nullptr) ? AnimationSequenceBase->GetName() : "Invalid Animation sequence";
+		UE_LOG(LogAnimationBlueprintLibrary, Warning, TEXT("Raw Animation Track name %s does not exist in Animation Sequence %s"), *TrackName.ToString(), *AnimSequenceName);
+	}
+}
+
+void UAnimationBlueprintLibrary::GetRawTrackScaleData(const UAnimSequenceBase* AnimationSequenceBase, const FName TrackName, TArray<FVector>& ScaleData)
+{
+	ScaleData.Empty();
+	if (IsValidRawAnimationTrackName(AnimationSequenceBase, TrackName))
+	{
+		const FRawAnimSequenceTrack& RawTrack = GetRawAnimationTrackByName(AnimationSequenceBase, TrackName);
+		ScaleData.Append(RawTrack.ScaleKeys);
+	}
+	else
+	{
+		const FString AnimSequenceName = (AnimationSequenceBase != nullptr) ? AnimationSequenceBase->GetName() : "Invalid Animation sequence";
+		UE_LOG(LogAnimationBlueprintLibrary, Warning, TEXT("Raw Animation Track name %s does not exist in Animation Sequence %s"), *TrackName.ToString(), *AnimSequenceName);
+	}
+}
+
+void UAnimationBlueprintLibrary::GetRawTrackData(const UAnimSequenceBase* AnimationSequenceBase, const FName TrackName, TArray<FVector>& PositionKeys, TArray<FQuat>& RotationKeys, TArray<FVector>& ScalingKeys)
+{
+	PositionKeys.Empty();
+	RotationKeys.Empty();
+	ScalingKeys.Empty();
+	if (IsValidRawAnimationTrackName(AnimationSequenceBase, TrackName))
+	{		
+		const FRawAnimSequenceTrack& RawTrack = GetRawAnimationTrackByName(AnimationSequenceBase, TrackName);
+		PositionKeys.Append(RawTrack.PosKeys);
+		RotationKeys.Append(RawTrack.RotKeys);
+		ScalingKeys.Append(RawTrack.ScaleKeys);
+	}
+	else
+	{
+		const FString AnimSequenceName = (AnimationSequenceBase != nullptr) ? AnimationSequenceBase->GetName() : "Invalid Animation sequence";
+		UE_LOG(LogAnimationBlueprintLibrary, Warning, TEXT("Raw Animation Track name %s does not exist in Animation Sequence %s"), *TrackName.ToString(), *AnimSequenceName);
+	}
+}
+
+bool UAnimationBlueprintLibrary::IsValidRawAnimationTrackName(const UAnimSequenceBase* AnimationSequenceBase, const FName TrackName)
+{
+	bool bValidName = false;
+
+	if (AnimationSequenceBase)
+	{
+		const int32 TrackIndex = AnimationSequenceBase->GetDataModel()->GetBoneTrackIndexByName(TrackName);
+		bValidName = (TrackIndex != INDEX_NONE);
+	}
+	else
+	{
+		UE_LOG(LogAnimationBlueprintLibrary, Warning, TEXT("Invalid Animation Sequence supplied for IsValidRawAnimationTrackName"));
+	}
+
+	return bValidName;
+}
+
+const FRawAnimSequenceTrack&  UAnimationBlueprintLibrary::GetRawAnimationTrackByName(const UAnimSequenceBase* AnimationSequenceBase, const FName TrackName)
+{
+	checkf(AnimationSequenceBase, TEXT("Invalid Animation Sequence supplied for GetRawAnimationTrackByName"));
+
+	const FBoneAnimationTrack& AnimationTrack = AnimationSequenceBase->GetDataModel()->GetBoneTrackByName(TrackName);
+	return AnimationTrack.InternalTrackData;
+}
+
+void UAnimationBlueprintLibrary::GetBoneCompressionSettings(const UAnimSequence* AnimationSequence, UAnimBoneCompressionSettings*& CompressionSettings)
+{
+	if (AnimationSequence == nullptr)
+	{
+		UE_LOG(LogAnimationBlueprintLibrary, Warning, TEXT("Invalid Animation Sequence supplied for GetBoneCompressionSettings"));
+		return;
+	}
+
+	CompressionSettings = AnimationSequence->BoneCompressionSettings;
+}
+
+void UAnimationBlueprintLibrary::SetBoneCompressionSettings(UAnimSequence* AnimationSequence, UAnimBoneCompressionSettings* CompressionSettings)
+{
+	if (AnimationSequence == nullptr)
+	{
+		UE_LOG(LogAnimationBlueprintLibrary, Warning, TEXT("Invalid Animation Sequence supplied for SetBoneCompressionSettings"));
+		return;
+	}
+
+	if (CompressionSettings == nullptr || !CompressionSettings->AreSettingsValid())
+	{
+		UE_LOG(LogAnimationBlueprintLibrary, Warning, TEXT("Invalid Bone Compression Settings supplied for SetBoneCompressionSettings"));
+		return;
+	}
+
+	AnimationSequence->BoneCompressionSettings = CompressionSettings;
+}
+
+void UAnimationBlueprintLibrary::GetCurveCompressionSettings(const UAnimSequence* AnimationSequence, UAnimCurveCompressionSettings*& CompressionSettings)
+{
+	if (AnimationSequence == nullptr)
+	{
+		UE_LOG(LogAnimationBlueprintLibrary, Warning, TEXT("Invalid Animation Sequence supplied for GetCurveCompressionSettings"));
+		return;
+	}
+
+	CompressionSettings = AnimationSequence->CurveCompressionSettings;
+}
+
+void UAnimationBlueprintLibrary::SetCurveCompressionSettings(UAnimSequence* AnimationSequence, UAnimCurveCompressionSettings* CompressionSettings)
+{
+	if (AnimationSequence == nullptr)
+	{
+		UE_LOG(LogAnimationBlueprintLibrary, Warning, TEXT("Invalid Animation Sequence supplied for SetCurveCompressionSettings"));
+		return;
+	}
+
+	if (CompressionSettings == nullptr || !CompressionSettings->AreSettingsValid())
+	{
+		UE_LOG(LogAnimationBlueprintLibrary, Warning, TEXT("Invalid Bone Compression Settings supplied for SetCurveCompressionSettings"));
+		return;
+	}
+
+	AnimationSequence->CurveCompressionSettings = CompressionSettings;
+}
+
+void UAnimationBlueprintLibrary::GetAdditiveAnimationType(const UAnimSequence* AnimationSequence, TEnumAsByte<enum EAdditiveAnimationType>& AdditiveAnimationType)
+{
+	if (AnimationSequence)
+	{
+		AdditiveAnimationType = AnimationSequence->AdditiveAnimType;
+	}
+	else
+	{
+		UE_LOG(LogAnimationBlueprintLibrary, Warning, TEXT("Invalid Animation Sequence supplied for GetAdditiveAnimationType"));
+	}
+}
+
+void UAnimationBlueprintLibrary::SetAdditiveAnimationType(UAnimSequence* AnimationSequence, const TEnumAsByte<enum EAdditiveAnimationType> AdditiveAnimationType)
+{
+	if (AnimationSequence)
+	{
+		AnimationSequence->AdditiveAnimType = AdditiveAnimationType;
+	}
+	else
+	{
+		UE_LOG(LogAnimationBlueprintLibrary, Warning, TEXT("Invalid Animation Sequence supplied for SetAdditiveAnimationType"));
+	}
+}
+
+void UAnimationBlueprintLibrary::GetAdditiveBasePoseType(const UAnimSequence* AnimationSequence, TEnumAsByte<enum EAdditiveBasePoseType>& AdditiveBasePoseType)
+{
+	if (AnimationSequence)
+	{
+		AdditiveBasePoseType = AnimationSequence->RefPoseType;
+	}
+	else
+	{
+		UE_LOG(LogAnimationBlueprintLibrary, Warning, TEXT("Invalid Animation Sequence supplied for GetAdditiveBasePoseType"));
+	}
+}
+
+void UAnimationBlueprintLibrary::SetAdditiveBasePoseType(UAnimSequence* AnimationSequence, const TEnumAsByte<enum EAdditiveBasePoseType> AdditiveBasePoseType)
+{
+	if (AnimationSequence)
+	{
+		AnimationSequence->RefPoseType = AdditiveBasePoseType;
+	}
+	else
+	{
+		UE_LOG(LogAnimationBlueprintLibrary, Warning, TEXT("Invalid Animation Sequence supplied for SetAdditiveBasePoseType"));
+	}
+}
+
+void UAnimationBlueprintLibrary::GetAnimationInterpolationType(const UAnimSequence* AnimationSequence, EAnimInterpolationType& InterpolationType)
+{
+	if (AnimationSequence)
+	{
+		InterpolationType = AnimationSequence->Interpolation;
+	}
+	else
+	{
+		UE_LOG(LogAnimationBlueprintLibrary, Warning, TEXT("Invalid Animation Sequence supplied for GetAnimationInterpolationType"));
+	}
+}
+
+void UAnimationBlueprintLibrary::SetAnimationInterpolationType(UAnimSequence* AnimationSequence, EAnimInterpolationType Type)
+{
+	if (AnimationSequence)
+	{
+		AnimationSequence->Interpolation = Type;
+	}
+	else
+	{
+		UE_LOG(LogAnimationBlueprintLibrary, Warning, TEXT("Invalid Animation Sequence supplied for SetAnimationInterpolationType"));
+	}
+}
+
+bool UAnimationBlueprintLibrary::IsRootMotionEnabled(const UAnimSequence* AnimationSequence)
+{
+	bool bEnabled = false;
+	if (AnimationSequence)
+	{
+		bEnabled = AnimationSequence->bEnableRootMotion;
+	}
+	else
+	{
+		UE_LOG(LogAnimationBlueprintLibrary, Warning, TEXT("Invalid Animation Sequence supplied for IsRootMotionEnabled"));
+	}
+
+	return bEnabled;
+}
+
+void UAnimationBlueprintLibrary::SetRootMotionEnabled(UAnimSequence* AnimationSequence, bool bEnabled)
+{
+	if (AnimationSequence)
+	{
+		AnimationSequence->bEnableRootMotion = bEnabled;
+	}
+	else
+	{
+		UE_LOG(LogAnimationBlueprintLibrary, Warning, TEXT("Invalid Animation Sequence supplied for SetRootMotionEnabled"));
+	}
+}
+
+void  UAnimationBlueprintLibrary::GetRootMotionLockType(const UAnimSequence* AnimationSequence, TEnumAsByte<ERootMotionRootLock::Type>& LockType)
+{
+	if (AnimationSequence)
+	{
+		LockType = AnimationSequence->RootMotionRootLock;
+	}
+	else
+	{
+		UE_LOG(LogAnimationBlueprintLibrary, Warning, TEXT("Invalid Animation Sequence for GetRootMotionLockType"));
+	}
+}
+
+void UAnimationBlueprintLibrary::SetRootMotionLockType(UAnimSequence* AnimationSequence, TEnumAsByte<ERootMotionRootLock::Type> RootMotionLockType)
+{
+	if (AnimationSequence)
+	{
+		AnimationSequence->RootMotionRootLock = RootMotionLockType;
+	}
+	else
+	{
+		UE_LOG(LogAnimationBlueprintLibrary, Warning, TEXT("Invalid Animation Sequence for SetRootMotionLockType"));
+	}	
+}
+
+bool UAnimationBlueprintLibrary::IsRootMotionLockForced(const UAnimSequence* AnimationSequence)
+{
+	bool bIsLocked = false;
+	if (AnimationSequence)
+	{
+		bIsLocked = AnimationSequence->bForceRootLock;
+	}
+	else
+	{
+		UE_LOG(LogAnimationBlueprintLibrary, Warning, TEXT("Invalid Animation Sequence for IsRootMotionLockForced"));	
+	}
+
+	return bIsLocked;
+}
+
+void UAnimationBlueprintLibrary::SetIsRootMotionLockForced(UAnimSequence* AnimationSequence, bool bForced)
+{
+	if (AnimationSequence)
+	{
+		AnimationSequence->bForceRootLock = bForced;
+	}
+	else
+	{
+		UE_LOG(LogAnimationBlueprintLibrary, Warning, TEXT("Invalid Animation Sequence for SetIsRootMotionLockForced"));
+	}
+}
+
+void UAnimationBlueprintLibrary::GetAnimationSyncMarkers(const UAnimSequence* AnimationSequence, TArray<FAnimSyncMarker>& Markers)
+{
+	Markers.Empty();
+	if (AnimationSequence)
+	{
+		Markers = AnimationSequence->AuthoredSyncMarkers;;
+	}
+	else
+	{
+		UE_LOG(LogAnimationBlueprintLibrary, Warning, TEXT("Invalid Animation Sequence for GetAnimationSyncMarkers"));
+	}
+}
+
+void UAnimationBlueprintLibrary::GetUniqueMarkerNames(const UAnimSequence* AnimationSequence, TArray<FName>& MarkerNames)
+{
+	MarkerNames.Empty();
+	if (AnimationSequence)
+	{
+		MarkerNames = AnimationSequence->UniqueMarkerNames;
+	}
+	else
+	{
+		UE_LOG(LogAnimationBlueprintLibrary, Warning, TEXT("Invalid Animation Sequence for GetUniqueMarkerNames"));
+	}
+}
+
+void UAnimationBlueprintLibrary::AddAnimationSyncMarker(UAnimSequence* AnimationSequence, FName MarkerName, float Time, FName TrackName)
+{
+	if (AnimationSequence)
+	{
+		const bool bIsValidTrackName = IsValidAnimNotifyTrackName(AnimationSequence, TrackName);
+		const bool bIsValidTime = IsValidTimeInternal(AnimationSequence, Time);
+
+		if (bIsValidTrackName && bIsValidTime)
+		{
+			FAnimSyncMarker NewMarker;
+			NewMarker.MarkerName = MarkerName;
+			NewMarker.Time = Time;
+			NewMarker.TrackIndex = GetTrackIndexForAnimationNotifyTrackName(AnimationSequence, TrackName);
+
+			AnimationSequence->AuthoredSyncMarkers.Add(NewMarker);
+			AnimationSequence->AnimNotifyTracks[NewMarker.TrackIndex].SyncMarkers.Add(&AnimationSequence->AuthoredSyncMarkers.Last());
+						
+			AnimationSequence->RefreshSyncMarkerDataFromAuthored();
+
+			// Refresh all cached data
+			AnimationSequence->RefreshCacheData();
+		}
+		else
+		{
+			if (!bIsValidTrackName)
+			{
+				UE_LOG(LogAnimationBlueprintLibrary, Warning, TEXT("Animation Notify Track %s does not exist in Animation Sequence %s"), *TrackName.ToString(), *AnimationSequence->GetName());
+			}
+
+			if (!bIsValidTime)
+			{
+				UE_LOG(LogAnimationBlueprintLibrary, Warning, TEXT("%f is side of Animation Sequence %s range"), Time, *AnimationSequence->GetName());
+			}
+		}
+	}
+	else
+	{
+		UE_LOG(LogAnimationBlueprintLibrary, Warning, TEXT("Invalid Animation Sequence for AddAnimationSyncMarker"));
+	}
+	
+}
+
+bool UAnimationBlueprintLibrary::IsValidAnimationSyncMarkerName(const UAnimSequence* AnimationSequence, FName MarkerName)
+{
+	bool bIsValid = false;
+	if (AnimationSequence)
+	{
+		bIsValid = AnimationSequence->UniqueMarkerNames.Contains(MarkerName);
+	}
+	else
+	{
+		UE_LOG(LogAnimationBlueprintLibrary, Warning, TEXT("Invalid Animation Sequence for IsValidAnimationSyncMarkerName"));
+	}
+
+	return bIsValid;	
+}
+
+int32 UAnimationBlueprintLibrary::RemoveAnimationSyncMarkersByName(UAnimSequence* AnimationSequence, FName MarkerName)
+{
+	int32 NumRemovedMarkers = 0;
+	if (AnimationSequence)
+	{
+		NumRemovedMarkers = AnimationSequence->AuthoredSyncMarkers.RemoveAll(
+			[&](const FAnimSyncMarker& Marker)
+		{
+			return Marker.MarkerName == MarkerName;
+		});
+
+		AnimationSequence->RefreshSyncMarkerDataFromAuthored();
+
+		// Refresh all cached data
+		AnimationSequence->RefreshCacheData();
+	}
+	else
+	{
+		UE_LOG(LogAnimationBlueprintLibrary, Warning, TEXT("Invalid Animation Sequence for RemoveAnimationSyncMarkersByName"));
+	}
+	
+	return NumRemovedMarkers;
+}
+
+int32 UAnimationBlueprintLibrary::RemoveAnimationSyncMarkersByTrack(UAnimSequence* AnimationSequence, FName NotifyTrackName)
+{
+	int32 NumRemovedMarkers = 0;
+	if (AnimationSequence)
+	{
+		const int32 TrackIndex = GetTrackIndexForAnimationNotifyTrackName(AnimationSequence, NotifyTrackName);
+		if (TrackIndex != INDEX_NONE)
+		{
+			NumRemovedMarkers = AnimationSequence->AuthoredSyncMarkers.RemoveAll(
+				[&](const FAnimSyncMarker& Marker)
+			{
+				return Marker.TrackIndex == TrackIndex;
+			});
+
+			AnimationSequence->RefreshSyncMarkerDataFromAuthored();
+
+			// Refresh all cached data
+			AnimationSequence->RefreshCacheData();
+		}
+		else
+		{
+			UE_LOG(LogAnimationBlueprintLibrary, Warning, TEXT("Animation Notify Track %s does not exist on Animation Sequence %s"), *NotifyTrackName.ToString(), *AnimationSequence->GetName());
+		}
+	}
+	else
+	{
+		UE_LOG(LogAnimationBlueprintLibrary, Warning, TEXT("Invalid Animation Sequence for RemoveAnimationSyncMarkersByTrack"));
+	}	
+	
+	return NumRemovedMarkers;
+}
+
+void UAnimationBlueprintLibrary::RemoveAllAnimationSyncMarkers(UAnimSequence* AnimationSequence)
+{
+	if (AnimationSequence)
+	{
+		AnimationSequence->AuthoredSyncMarkers.Empty();
+		AnimationSequence->RefreshSyncMarkerDataFromAuthored();
+
+		// Refresh all cached data
+		AnimationSequence->RefreshCacheData();
+	}
+	else
+	{
+		UE_LOG(LogAnimationBlueprintLibrary, Warning, TEXT("Invalid Animation Sequence for RemoveAllAnimationSyncMarkers"));
+	}
+}
+
+void UAnimationBlueprintLibrary::GetAnimationNotifyEvents(const UAnimSequenceBase* AnimationSequenceBase, TArray<FAnimNotifyEvent>& NotifyEvents)
+{
+	NotifyEvents.Empty();
+	if (AnimationSequenceBase)
+	{
+		NotifyEvents = AnimationSequenceBase->Notifies;
+	}
+	else
+	{
+		UE_LOG(LogAnimationBlueprintLibrary, Warning, TEXT("Invalid Animation Sequence for GetAnimationNotifyEvents"));
+	}	
+}
+
+void UAnimationBlueprintLibrary::GetAnimationNotifyEventNames(const UAnimSequenceBase* AnimationSequenceBase, TArray<FName>& EventNames)
+{
+	EventNames.Empty();
+	if (AnimationSequenceBase)
+	{
+		for (const FAnimNotifyEvent& Event : AnimationSequenceBase->Notifies)
+		{
+			EventNames.AddUnique(Event.NotifyName);
+		}
+	}
+	else
+	{
+		UE_LOG(LogAnimationBlueprintLibrary, Warning, TEXT("Invalid Animation Sequence for GetAnimationNotifyEventNames"));
+	}	
+}
+
+UAnimNotify* UAnimationBlueprintLibrary::AddAnimationNotifyEvent(UAnimSequenceBase* AnimationSequenceBase, FName NotifyTrackName, float StartTime, TSubclassOf<UAnimNotify> NotifyClass)
+{
+	UAnimNotify* Notify = nullptr;
+	if (AnimationSequenceBase)
+	{
+		const bool bIsValidTrackName = IsValidAnimNotifyTrackName(AnimationSequenceBase, NotifyTrackName);
+		const bool bIsValidTime = IsValidTimeInternal(AnimationSequenceBase, StartTime);
+
+		if (bIsValidTrackName && bIsValidTime)
+		{
+			FAnimNotifyEvent& NewEvent = AnimationSequenceBase->Notifies.AddDefaulted_GetRef();
+
+			NewEvent.NotifyName = NAME_None;
+			NewEvent.Link(AnimationSequenceBase, StartTime);
+			NewEvent.TriggerTimeOffset = GetTriggerTimeOffsetForType(AnimationSequenceBase->CalculateOffsetForNotify(StartTime));
+			NewEvent.TrackIndex = GetTrackIndexForAnimationNotifyTrackName(AnimationSequenceBase, NotifyTrackName);
+			NewEvent.NotifyStateClass = nullptr;
+			NewEvent.Guid = FGuid::NewGuid();
+
+			if (NotifyClass)
+			{
+				Notify = NewObject<UAnimNotify>(AnimationSequenceBase, NotifyClass, NAME_None, RF_Transactional);
+				NewEvent.Notify = Notify;
+
+				// Setup name for new event
+				if(NewEvent.Notify)
+				{
+					NewEvent.NotifyName = FName(*NewEvent.Notify->GetNotifyName());
+				}
+			}
+			else
+			{
+				NewEvent.Notify = nullptr;
+			}
+
+			// Refresh all cached data
+			AnimationSequenceBase->RefreshCacheData();
+		}
+		else
+		{
+			if (!bIsValidTrackName)
+			{
+				UE_LOG(LogAnimationBlueprintLibrary, Warning, TEXT("Animation Notify Track %s does not exist on Animation Sequence %s"), *NotifyTrackName.ToString(), *AnimationSequenceBase->GetName());
+			}
+
+			if (!bIsValidTime)
+			{
+				UE_LOG(LogAnimationBlueprintLibrary, Warning, TEXT("%f is side of Animation Sequence %s range"), StartTime, *AnimationSequenceBase->GetName());
+			}
+		}
+	}
+	else
+	{
+		UE_LOG(LogAnimationBlueprintLibrary, Warning, TEXT("Invalid Animation Sequence for AddAnimationNotifyEvent"));
+	}
+
+	return Notify;
+}
+
+UAnimNotifyState* UAnimationBlueprintLibrary::AddAnimationNotifyStateEvent(UAnimSequenceBase* AnimationSequenceBase, FName NotifyTrackName, float StartTime, float Duration, TSubclassOf<UAnimNotifyState> NotifyStateClass)
+{
+	UAnimNotifyState* NotifyState = nullptr;
+	if (AnimationSequenceBase)
+	{
+		const bool bIsValidTrackName = IsValidAnimNotifyTrackName(AnimationSequenceBase, NotifyTrackName);
+		const bool bIsValidTime = IsValidTimeInternal(AnimationSequenceBase, StartTime);
+
+		if (bIsValidTrackName && bIsValidTime)
+		{
+			FAnimNotifyEvent& NewEvent = AnimationSequenceBase->Notifies.AddDefaulted_GetRef();
+
+			NewEvent.NotifyName = NAME_None;
+			NewEvent.Link(AnimationSequenceBase, StartTime);
+			NewEvent.TriggerTimeOffset = GetTriggerTimeOffsetForType(AnimationSequenceBase->CalculateOffsetForNotify(StartTime));
+			NewEvent.TrackIndex = GetTrackIndexForAnimationNotifyTrackName(AnimationSequenceBase, NotifyTrackName);
+			NewEvent.Notify = nullptr;
+			NewEvent.Guid = FGuid::NewGuid();
+
+			if (NotifyStateClass)
+			{
+				NotifyState = NewObject<UAnimNotifyState>(AnimationSequenceBase, NotifyStateClass, NAME_None, RF_Transactional);
+				NewEvent.NotifyStateClass = NotifyState;
+
+				// Setup name and duration for new event
+				if (NewEvent.NotifyStateClass)
+				{
+					NewEvent.NotifyName = FName(*NewEvent.NotifyStateClass->GetNotifyName());
+					NewEvent.SetDuration(Duration);
+					NewEvent.EndLink.Link(AnimationSequenceBase, NewEvent.EndLink.GetTime());
+				}
+			}
+			else
+			{
+				NewEvent.NotifyStateClass = nullptr;
+			}
+
+			// Refresh all cached data
+			AnimationSequenceBase->RefreshCacheData();			
+		}
+		else
+		{
+			if (!bIsValidTrackName)
+			{
+				UE_LOG(LogAnimationBlueprintLibrary, Warning, TEXT("Animation Notify Track %s does not exist on Animation Sequence %s"), *NotifyTrackName.ToString(), *AnimationSequenceBase->GetName());
+			}
+
+			if (!bIsValidTime)
+			{
+				UE_LOG(LogAnimationBlueprintLibrary, Warning, TEXT("%f is side of Animation Sequence %s range"), StartTime, *AnimationSequenceBase->GetName());
+			}
+		}
+	}
+	else
+	{
+		UE_LOG(LogAnimationBlueprintLibrary, Warning, TEXT("Invalid Animation Sequence for AddAnimationNotifyStateEvent"));
+	}
+
+	return NotifyState;
+}
+
+void UAnimationBlueprintLibrary::AddAnimationNotifyEventObject(UAnimSequenceBase* AnimationSequenceBase, float StartTime, UAnimNotify* Notify, FName NotifyTrackName)
+{
+	if (AnimationSequenceBase)
+	{
+		const bool bValidNotify = Notify != nullptr;
+		const bool bValidOuter = bValidNotify && Notify->GetOuter() == AnimationSequenceBase;
+		const bool bIsValidTrackName = IsValidAnimNotifyTrackName(AnimationSequenceBase, NotifyTrackName);
+		const bool bIsValidTime = IsValidTimeInternal(AnimationSequenceBase, StartTime);
+
+		if (bValidNotify && bValidOuter && bIsValidTrackName && bIsValidTime)
+		{
+			FAnimNotifyEvent& NewEvent = AnimationSequenceBase->Notifies.AddDefaulted_GetRef();
+
+			NewEvent.NotifyName = NAME_None;
+			NewEvent.Link(AnimationSequenceBase, StartTime);
+			NewEvent.TriggerTimeOffset = GetTriggerTimeOffsetForType(AnimationSequenceBase->CalculateOffsetForNotify(StartTime));
+			NewEvent.TrackIndex = GetTrackIndexForAnimationNotifyTrackName(AnimationSequenceBase, NotifyTrackName);
+			NewEvent.NotifyStateClass = nullptr;
+			NewEvent.Notify = Notify;
+			NewEvent.Guid = FGuid::NewGuid();		
+
+			// Refresh all cached data
+			AnimationSequenceBase->RefreshCacheData();
+		}
+		else
+		{
+			if (!bValidNotify)
+			{
+				UE_LOG(LogAnimationBlueprintLibrary, Warning, TEXT("Invalid Animation Notify in AddAnimationNotifyEventObject"));
+			}
+
+			if (!bValidOuter)
+			{
+				const FString NotifyName = bValidNotify ? Notify->GetName() : "Invalid Notify";
+				UE_LOG(LogAnimationBlueprintLibrary, Warning, TEXT("Animation Notify %s Outer is not %s"), *NotifyName, *AnimationSequenceBase->GetName());
+			}
+
+			if (!bIsValidTrackName)
+			{
+				UE_LOG(LogAnimationBlueprintLibrary, Warning, TEXT("Animation Notify Track %s does not exist on Animation Sequence %s"), *NotifyTrackName.ToString(), *AnimationSequenceBase->GetName());
+			}
+
+			if (!bIsValidTime)
+			{
+				UE_LOG(LogAnimationBlueprintLibrary, Warning, TEXT("%f is side of Animation Sequence %s range"), StartTime, *AnimationSequenceBase->GetName());
+			}
+		}
+	}
+	else
+	{
+		UE_LOG(LogAnimationBlueprintLibrary, Warning, TEXT("Invalid Animation Sequence for AddAnimationNotifyEventObject"));
+	}
+}
+
+void UAnimationBlueprintLibrary::AddAnimationNotifyStateEventObject(UAnimSequenceBase* AnimationSequenceBase, float StartTime, float Duration, UAnimNotifyState* NotifyState, FName NotifyTrackName)
+{
+	if (AnimationSequenceBase)
+	{
+		const bool bValidNotifyState = NotifyState != nullptr;
+		const bool bValidOuter = bValidNotifyState && NotifyState->GetOuter() == AnimationSequenceBase;
+		const bool bIsValidTrackName = IsValidAnimNotifyTrackName(AnimationSequenceBase, NotifyTrackName);
+		const bool bIsValidTime = IsValidTimeInternal(AnimationSequenceBase, StartTime);
+
+		if (bValidNotifyState && bValidOuter && bIsValidTrackName && bIsValidTime)
+		{
+			FAnimNotifyEvent& NewEvent = AnimationSequenceBase->Notifies.AddDefaulted_GetRef();
+
+			NewEvent.NotifyName = NAME_None;
+			NewEvent.Link(AnimationSequenceBase, StartTime);
+			NewEvent.TriggerTimeOffset = GetTriggerTimeOffsetForType(AnimationSequenceBase->CalculateOffsetForNotify(StartTime));
+			NewEvent.TrackIndex = GetTrackIndexForAnimationNotifyTrackName(AnimationSequenceBase, NotifyTrackName);
+			NewEvent.NotifyStateClass = NotifyState;
+			NewEvent.Notify = nullptr;
+			NewEvent.SetDuration(Duration);
+			NewEvent.EndLink.Link(AnimationSequenceBase, NewEvent.EndLink.GetTime());
+			NewEvent.Guid = FGuid::NewGuid();
+
+			// Refresh all cached data
+			AnimationSequenceBase->RefreshCacheData();
+		}
+		else
+		{
+			if (!bValidNotifyState)
+			{
+				UE_LOG(LogAnimationBlueprintLibrary, Warning, TEXT("Invalid Animation Notify in AddAnimationNotifyEventObject"));
+			}
+
+			if (!bValidOuter)
+			{
+				const FString NotifyName = bValidNotifyState ? NotifyState->GetName() : "Invalid Notify";
+				UE_LOG(LogAnimationBlueprintLibrary, Warning, TEXT("Animation Notify %s Outer is not %s"), *NotifyName, *AnimationSequenceBase->GetName());
+			}
+
+			if (!bIsValidTrackName)
+			{
+				UE_LOG(LogAnimationBlueprintLibrary, Warning, TEXT("Animation Notify Track %s does not exist on Animation Sequence %s"), *NotifyTrackName.ToString(), *AnimationSequenceBase->GetName());
+			}
+
+			if (!bIsValidTime)
+			{
+				UE_LOG(LogAnimationBlueprintLibrary, Warning, TEXT("%f is side of Animation Sequence %s range"), StartTime, *AnimationSequenceBase->GetName());
+			}
+		}
+	}
+	else
+	{
+		UE_LOG(LogAnimationBlueprintLibrary, Warning, TEXT("Invalid Animation Sequence for AddAnimationNotifyEventObject"));
+	}
+}
+
+static void ReplaceAnimNotifies_Helper(UAnimSequenceBase* AnimationSequence, UClass* OldNotifyClass, UClass* NewNotifyClass, FOnNotifyReplaced OnNotifyReplaced, FOnNotifyStateReplaced OnNotifyStateReplaced)
+{
+	if (AnimationSequence)
+	{
+		if (OldNotifyClass != nullptr && NewNotifyClass != nullptr)
+		{
+			bool bModified = false;
+			for(int32 NotifyIndex = 0; NotifyIndex < AnimationSequence->Notifies.Num(); ++NotifyIndex)
+			{
+				FAnimNotifyEvent& NotifyEvent = AnimationSequence->Notifies[NotifyIndex];
+
+				if ((NotifyEvent.Notify && NotifyEvent.Notify->GetClass() == OldNotifyClass) || 
+					(NotifyEvent.NotifyStateClass && NotifyEvent.NotifyStateClass->GetClass() == OldNotifyClass))
+				{
+
+					bModified = true;
+
+					// Copy relevant data from the old notify
+					float StartTime = NotifyEvent.GetTime();
+					float Length = NotifyEvent.GetDuration();
+					int32 TargetTrackIndex = NotifyEvent.TrackIndex;
+					float TriggerTimeOffset = NotifyEvent.TriggerTimeOffset;
+					float EndTriggerTimeOffset = NotifyEvent.EndTriggerTimeOffset;
+					int32 SlotIndex = NotifyEvent.GetSlotIndex();
+					int32 EndSlotIndex = NotifyEvent.EndLink.GetSlotIndex();
+					int32 SegmentIndex = NotifyEvent.GetSegmentIndex();
+					int32 EndSegmentIndex = NotifyEvent.GetSegmentIndex();
+					EAnimLinkMethod::Type LinkMethod = NotifyEvent.GetLinkMethod();
+					EAnimLinkMethod::Type EndLinkMethod = NotifyEvent.EndLink.GetLinkMethod();
+					UAnimNotify* OldNotify = NotifyEvent.Notify;
+					UAnimNotifyState* OldNotifyState = NotifyEvent.NotifyStateClass;
+
+					// Remove old notify
+					AnimationSequence->Notifies.RemoveAt(NotifyIndex, 1, false);
+
+					// Add new notify in old notifies place
+					AnimationSequence->Notifies.InsertDefaulted(NotifyIndex);
+					FAnimNotifyEvent& NewEvent = AnimationSequence->Notifies[NotifyIndex];
+
+					// Setup new notify
+					NewEvent.NotifyName = NAME_None;
+					NewEvent.Link(AnimationSequence, StartTime);
+					NewEvent.TriggerTimeOffset = TriggerTimeOffset;
+					NewEvent.TrackIndex = TargetTrackIndex;
+					NewEvent.ChangeSlotIndex(SlotIndex);
+					NewEvent.SetSegmentIndex(SegmentIndex);
+					NewEvent.ChangeLinkMethod(LinkMethod);
+					NewEvent.Guid = FGuid::NewGuid();
+
+					UObject* AnimNotifyClass = NewObject<UObject>(AnimationSequence, NewNotifyClass, NAME_None, RF_Transactional);
+					NewEvent.NotifyStateClass = Cast<UAnimNotifyState>(AnimNotifyClass);
+					NewEvent.Notify = Cast<UAnimNotify>(AnimNotifyClass);
+
+					// Setup name and duration for new event
+					if (NewEvent.NotifyStateClass)
+					{
+						NewEvent.NotifyName = FName(*NewEvent.NotifyStateClass->GetNotifyName());
+						NewEvent.SetDuration(Length);
+						NewEvent.EndTriggerTimeOffset = EndTriggerTimeOffset;
+						NewEvent.EndLink.ChangeSlotIndex(EndSlotIndex);
+						NewEvent.EndLink.SetSegmentIndex(EndSegmentIndex);
+						NewEvent.EndLink.ChangeLinkMethod(EndLinkMethod);
+
+						OnNotifyStateReplaced.ExecuteIfBound(OldNotifyState, NewEvent.NotifyStateClass);
+					}
+					else if(NewEvent.Notify)
+					{
+						NewEvent.NotifyName = FName(*NewEvent.Notify->GetNotifyName());
+
+						OnNotifyReplaced.ExecuteIfBound(OldNotify, NewEvent.Notify);
+					}
+
+					NewEvent.Update();
+				}
+			}
+
+			if(bModified)
+			{
+				// Refresh all cached data
+				AnimationSequence->MarkPackageDirty();
+				AnimationSequence->RefreshCacheData();
+			}
+		}
+		else
+		{
+			UE_LOG(LogAnimationBlueprintLibrary, Warning, TEXT("Invalid Notify Class for ReplaceAnimNotifies"));
+		}
+	}
+	else
+	{
+		UE_LOG(LogAnimationBlueprintLibrary, Warning, TEXT("Invalid Animation Sequence for ReplaceAnimNotifies"));
+	}
+}
+
+void UAnimationBlueprintLibrary::ReplaceAnimNotifyStates(UAnimSequenceBase* AnimationSequenceBase, TSubclassOf<UAnimNotifyState> OldNotifyClass, TSubclassOf<UAnimNotifyState> NewNotifyClass, FOnNotifyStateReplaced OnNotifyStateReplaced)
+{
+	ReplaceAnimNotifies_Helper(AnimationSequenceBase, OldNotifyClass.Get(), NewNotifyClass.Get(), FOnNotifyReplaced(), OnNotifyStateReplaced);
+}
+
+void UAnimationBlueprintLibrary::ReplaceAnimNotifies(UAnimSequenceBase* AnimationSequenceBase, TSubclassOf<UAnimNotify> OldNotifyClass, TSubclassOf<UAnimNotify> NewNotifyClass, FOnNotifyReplaced OnNotifyReplaced)
+{
+	ReplaceAnimNotifies_Helper(AnimationSequenceBase, OldNotifyClass.Get(), NewNotifyClass.Get(), OnNotifyReplaced, FOnNotifyStateReplaced());
+}
+
+void UAnimationBlueprintLibrary::CopyAnimNotifiesFromSequence(UAnimSequenceBase* SourceAnimationSequenceBase, UAnimSequenceBase* DestinationAnimSequenceBase)
+{
+	if (SourceAnimationSequenceBase == nullptr)
+	{
+		UE_LOG(LogAnimationBlueprintLibrary, Warning, TEXT("Invalid Source Animation Sequence for CopyAnimNotifiesFromSequence"));
+	}
+	else if (DestinationAnimSequenceBase == nullptr)
+	{
+		UE_LOG(LogAnimationBlueprintLibrary, Warning, TEXT("Invalid Destination Animation Sequence for CopyAnimNotifiesFromSequence"));
+	}
+	else if (SourceAnimationSequenceBase == DestinationAnimSequenceBase)
+	{
+		UE_LOG(LogAnimationBlueprintLibrary, Warning, TEXT("Source and Destination Animation Sequence are the same for CopyAnimNotifiesFromSequence"));
+	}
+	else
+	{
+		const bool bShowDialogs = false;
+		UE::Anim::CopyNotifies(SourceAnimationSequenceBase, DestinationAnimSequenceBase, bShowDialogs);
+	}
+}
+
+int32 UAnimationBlueprintLibrary::RemoveAnimationNotifyEventsByName(UAnimSequenceBase* AnimationSequenceBase, FName NotifyName)
+{
+	int32 NumRemovedEvents = 0;
+	if (AnimationSequenceBase)
+	{
+		NumRemovedEvents = AnimationSequenceBase->Notifies.RemoveAll(
+			[&](const FAnimNotifyEvent& Event)
+		{
+			return Event.NotifyName == NotifyName;
+		});
+
+		// Refresh all cached data
+		AnimationSequenceBase->RefreshCacheData();
+	}
+	else
+	{
+		UE_LOG(LogAnimationBlueprintLibrary, Warning, TEXT("Invalid Animation Sequence for RemoveAnimationNotifyEventsByName"));
+	}	
+
+	return NumRemovedEvents;
+}
+
+int32 UAnimationBlueprintLibrary::RemoveAnimationNotifyEventsByTrack(UAnimSequenceBase* AnimationSequenceBase, FName NotifyTrackName)
+{
+	int32 NumRemovedEvents = 0;
+	if (AnimationSequenceBase)
+	{
+		const bool bIsValidTrackName = IsValidAnimNotifyTrackName(AnimationSequenceBase, NotifyTrackName);
+		if (bIsValidTrackName)
+		{
+			const int32 TrackIndex = GetTrackIndexForAnimationNotifyTrackName(AnimationSequenceBase, NotifyTrackName);
+			NumRemovedEvents = AnimationSequenceBase->Notifies.RemoveAll(
+				[&](const FAnimNotifyEvent& Event)
+			{
+				return Event.TrackIndex == TrackIndex;
+			});
+
+			// Refresh all cached data
+			AnimationSequenceBase->RefreshCacheData();
+		}
+		else
+		{
+			UE_LOG(LogAnimationBlueprintLibrary, Warning, TEXT("Animation Notify Track %s does not exist on Animation Sequence %s"), *NotifyTrackName.ToString(), *AnimationSequenceBase->GetName());
+		}
+	}
+	else
+	{
+		UE_LOG(LogAnimationBlueprintLibrary, Warning, TEXT("Invalid Animation Sequence for RemoveAnimationNotifyEventsByTrack"));
+	}
+	
+
+	return NumRemovedEvents;
+}
+
+void UAnimationBlueprintLibrary::GetAnimationNotifyTrackNames(const UAnimSequenceBase* AnimationSequenceBase, TArray<FName>& TrackNames)
+{
+	TrackNames.Empty();
+	if (AnimationSequenceBase)
+	{
+		for (const FAnimNotifyTrack& Track : AnimationSequenceBase->AnimNotifyTracks)
+		{
+			TrackNames.AddUnique(Track.TrackName);
+		}
+	}
+	else
+	{
+		UE_LOG(LogAnimationBlueprintLibrary, Warning, TEXT("Invalid Animation Sequence for GetAnimationNotifyTrackNames"));
+	}
+}
+
+void UAnimationBlueprintLibrary::AddAnimationNotifyTrack(UAnimSequenceBase* AnimationSequenceBase, FName NotifyTrackName, FLinearColor TrackColor /*= FLinearColor::White*/)
+{
+	if (AnimationSequenceBase)
+	{
+		const bool bExistingTrackName = IsValidAnimNotifyTrackName(AnimationSequenceBase, NotifyTrackName);
+		if (!bExistingTrackName)
+		{
+			FAnimNotifyTrack NewTrack;
+			NewTrack.TrackName = NotifyTrackName;
+			NewTrack.TrackColor = TrackColor;
+			AnimationSequenceBase->AnimNotifyTracks.Add(NewTrack);
+
+			// Refresh all cached data
+			AnimationSequenceBase->RefreshCacheData();
+		}
+		else
+		{
+			UE_LOG(LogAnimationBlueprintLibrary, Warning, TEXT("Animation Notify Track %s already exists on Animation Sequence %s"), *NotifyTrackName.ToString(), *AnimationSequenceBase->GetName());
+		}
+	}
+	else
+	{
+		UE_LOG(LogAnimationBlueprintLibrary, Warning, TEXT("Invalid Animation Sequence for AddAnimationNotifyTrack"));
+	}
+	
+}
+
+void UAnimationBlueprintLibrary::RemoveAnimationNotifyTrack(UAnimSequenceBase* AnimationSequenceBase, FName NotifyTrackName)
+{
+	if (AnimationSequenceBase)
+	{
+		const int32 TrackIndexToDelete = GetTrackIndexForAnimationNotifyTrackName(AnimationSequenceBase, NotifyTrackName);
+		if (TrackIndexToDelete != INDEX_NONE)
+		{	
+			// Remove all notifies and sync markers on the to-delete-track
+			AnimationSequenceBase->Notifies.RemoveAll([&](const FAnimNotifyEvent& Notify) { return Notify.TrackIndex == TrackIndexToDelete; });
+
+			// Before track removal, make sure everything behind is fixed
+			for (FAnimNotifyEvent& Notify : AnimationSequenceBase->Notifies)
+			{
+				if (Notify.TrackIndex > TrackIndexToDelete)
+				{
+					Notify.TrackIndex = Notify.TrackIndex - 1;
+				}				
+			}
+
+			if (UAnimSequence* AnimationSequence = Cast<UAnimSequence>(AnimationSequenceBase))
+			{
+				AnimationSequence->AuthoredSyncMarkers.RemoveAll([&](const FAnimSyncMarker& Marker) { return Marker.TrackIndex == TrackIndexToDelete; });
+			for (FAnimSyncMarker& SyncMarker : AnimationSequence->AuthoredSyncMarkers)
+			{
+				if (SyncMarker.TrackIndex > TrackIndexToDelete)
+				{
+					SyncMarker.TrackIndex = SyncMarker.TrackIndex - 1;
+				}
+			}
+			}			
+			
+			// Delete the track itself
+			AnimationSequenceBase->AnimNotifyTracks.RemoveAt(TrackIndexToDelete);
+
+			// Refresh all cached data
+			AnimationSequenceBase->RefreshCacheData();
+		}		
+	}
+	else
+	{
+		UE_LOG(LogAnimationBlueprintLibrary, Warning, TEXT("Invalid Animation Sequence for RemoveAnimationNotifyTrack"));
+	}
+}
+
+void UAnimationBlueprintLibrary::RemoveAllAnimationNotifyTracks(UAnimSequenceBase* AnimationSequenceBase)
+{
+	if (AnimationSequenceBase)
+{
+		AnimationSequenceBase->Notifies.Empty();
+		if (UAnimSequence* AnimationSequence = Cast<UAnimSequence>(AnimationSequenceBase))
+	{
+		AnimationSequence->AuthoredSyncMarkers.Empty();
+		}
+
+		// Remove all but one notify tracks
+		AnimationSequenceBase->AnimNotifyTracks.SetNum(1);
+
+		// Also remove all stale notifies and sync markers from only track 
+		AnimationSequenceBase->AnimNotifyTracks[0].Notifies.Empty();
+		AnimationSequenceBase->AnimNotifyTracks[0].SyncMarkers.Empty();
+
+		// Refresh all cached data
+		AnimationSequenceBase->RefreshCacheData();
+	}
+	else
+	{
+		UE_LOG(LogAnimationBlueprintLibrary, Warning, TEXT("Invalid Animation Sequence for RemoveAllAnimationNotifyTracks"));
+	}	
+}
+
+bool UAnimationBlueprintLibrary::IsValidAnimNotifyTrackName(const UAnimSequenceBase* AnimationSequenceBase, FName NotifyTrackName)
+{
+	bool bIsValid = false;
+	if (AnimationSequenceBase)
+	{
+		bIsValid = GetTrackIndexForAnimationNotifyTrackName(AnimationSequenceBase, NotifyTrackName) != INDEX_NONE;
+	}
+	else
+	{
+		UE_LOG(LogAnimationBlueprintLibrary, Warning, TEXT("Invalid Animation Sequence for IsValidAnimNotifyTrackName"));
+	}
+
+	return bIsValid;	
+}
+
+int32 UAnimationBlueprintLibrary::GetTrackIndexForAnimationNotifyTrackName(const UAnimSequenceBase* AnimationSequenceBase, FName NotifyTrackName)
+{
+	return AnimationSequenceBase->AnimNotifyTracks.IndexOfByPredicate(
+		[&](const FAnimNotifyTrack& Track)
+	{
+		return Track.TrackName == NotifyTrackName;
+	});
+}
+
+const FAnimNotifyTrack& UAnimationBlueprintLibrary::GetNotifyTrackByName(const UAnimSequenceBase* AnimationSequenceBase, FName NotifyTrackName)
+{
+	const int32 TrackIndex = GetTrackIndexForAnimationNotifyTrackName(AnimationSequenceBase, NotifyTrackName);
+	checkf(TrackIndex != INDEX_NONE, TEXT("Notify Track %s does not exist on %s"), *NotifyTrackName.ToString(), *AnimationSequenceBase->GetName());
+	return AnimationSequenceBase->AnimNotifyTracks[TrackIndex];
+}
+
+float UAnimationBlueprintLibrary::GetAnimNotifyEventTriggerTime(const FAnimNotifyEvent& NotifyEvent)
+{
+	return NotifyEvent.GetTriggerTime();
+}
+
+float UAnimationBlueprintLibrary::GetAnimNotifyEventDuration(const FAnimNotifyEvent& NotifyEvent)
+{
+	return NotifyEvent.Duration;
+}
+
+void UAnimationBlueprintLibrary::GetAnimationSyncMarkersForTrack(const UAnimSequence* AnimationSequence, FName NotifyTrackName, TArray<FAnimSyncMarker>& Markers)
+{
+	Markers.Empty();
+	if (AnimationSequence)
+	{
+		const bool bIsValidTrackName = IsValidAnimNotifyTrackName(AnimationSequence, NotifyTrackName);
+
+		if (bIsValidTrackName)
+		{
+			const FAnimNotifyTrack& Track = GetNotifyTrackByName(AnimationSequence, NotifyTrackName);
+			Markers.Empty(Track.SyncMarkers.Num());
+			for (FAnimSyncMarker* Marker : Track.SyncMarkers)
+			{
+				Markers.Add(*Marker);
+			}
+		}
+		else
+		{
+			UE_LOG(LogAnimationBlueprintLibrary, Warning, TEXT("Animation Notify Track %s does not exist on Animation Sequence %s"), *NotifyTrackName.ToString(), *AnimationSequence->GetName());
+		}
+	}
+	else
+	{
+		UE_LOG(LogAnimationBlueprintLibrary, Warning, TEXT("Invalid Animation Sequence for AddVectorCurveKey"));
+	}
+}
+
+void UAnimationBlueprintLibrary::GetAnimationNotifyEventsForTrack(const UAnimSequenceBase* AnimationSequenceBase, FName NotifyTrackName, TArray<FAnimNotifyEvent>& Events)
+{
+	Events.Empty();
+	if (AnimationSequenceBase)
+	{
+		const bool bIsValidTrackName = IsValidAnimNotifyTrackName(AnimationSequenceBase, NotifyTrackName);
+
+		if (bIsValidTrackName)
+		{
+			const FAnimNotifyTrack& Track = GetNotifyTrackByName(AnimationSequenceBase, NotifyTrackName);
+			Events.Empty(Track.Notifies.Num());
+			for (FAnimNotifyEvent* Event : Track.Notifies)
+			{
+				Events.Add(*Event);
+			}
+		}
+		else
+		{
+			UE_LOG(LogAnimationBlueprintLibrary, Warning, TEXT("Animation Notify Track %s does not exist on Animation Sequence %s"), *NotifyTrackName.ToString(), *AnimationSequenceBase->GetName());
+		}
+	}
+	else
+	{
+		UE_LOG(LogAnimationBlueprintLibrary, Warning, TEXT("Invalid Animation Sequence for AddVectorCurveKey"));
+	}
+}
+
+void UAnimationBlueprintLibrary::AddCurve(UAnimSequence* AnimationSequence, FName CurveName, ERawCurveTrackTypes CurveType /*= RCT_Float*/, bool bMetaDataCurve /*= false*/)
+{
+	if (AnimationSequence)
+	{
+		static const ESmartNameContainerType ContainerForCurveType[(int32)ERawCurveTrackTypes::RCT_MAX] = { ESmartNameContainerType::SNCT_CurveMapping, ESmartNameContainerType::SNCT_CurveMapping, ESmartNameContainerType::SNCT_TrackCurveMapping };
+		const ESmartNameContainerType CurveContainer = ContainerForCurveType[(int32)CurveType];
+		const int32 CurveFlags = bMetaDataCurve ? AACF_Metadata : AACF_DefaultCurve;
+		
+		// Validate combination of curve types
+
+		// Only Float metadata curves are valid
+		const bool bValidMetaData = !bMetaDataCurve || (bMetaDataCurve && CurveType == ERawCurveTrackTypes::RCT_Float);
+		// Transform curves can only be added if the curve name exists as a bone on the skeleton
+		const bool bValidTransformCurveData = CurveType != ERawCurveTrackTypes::RCT_Transform || (AnimationSequence->GetSkeleton() && DoesBoneNameExistInternal(AnimationSequence->GetSkeleton(), CurveName));
+
+		if (bValidMetaData && bValidTransformCurveData )
+		{
+			// Add or retrieve the smartname
+			const FName Names[(int32)ESmartNameContainerType::SNCT_MAX] = { USkeleton::AnimCurveMappingName, USkeleton::AnimTrackCurveMappingName };
+			const bool bCurveAdded = AddCurveInternal(AnimationSequence, CurveName, Names[(int32)CurveContainer], CurveFlags, CurveType);
+
+			if (!bCurveAdded)
+			{
+				// Curve already existed
+				UE_LOG(LogAnimationBlueprintLibrary, Warning, TEXT("Curve %s already exists on the Skeleton %s."), *CurveName.ToString(), *AnimationSequence->GetSkeleton()->GetName());
+			}
+		}
+		else
+		{
+			if (!bValidMetaData)
+			{
+				UE_LOG(LogAnimationBlueprintLibrary, Warning, TEXT("Invalid Curve type to be create as metadata, currently only float curves are supported as metadata."));
+			}
+			
+			if (!bValidTransformCurveData)
+			{
+				UE_LOG(LogAnimationBlueprintLibrary, Warning, TEXT("Invalid Transform Curve name, the supplied name %s does not exist on the Skeleton %s."), *CurveName.ToString(), AnimationSequence->GetSkeleton() ? *AnimationSequence->GetSkeleton()->GetName() : TEXT("Invalid Skeleton"));
+			}
+		}
+	}
+	else
+	{
+		UE_LOG(LogAnimationBlueprintLibrary, Warning, TEXT("Invalid Animation Sequence supplied for AddCurve"));
+	}	
+}
+
+void UAnimationBlueprintLibrary::RemoveCurve(UAnimSequence* AnimationSequence, FName CurveName, bool bRemoveNameFromSkeleton /*= false*/)
+{
+	if (AnimationSequence)
+	{
+		const FName ContainerName = RetrieveContainerNameForCurve(AnimationSequence, CurveName);
+		if (ContainerName != NAME_None)
+		{
+			const bool bCurveRemoved = RemoveCurveInternal(AnimationSequence, CurveName, ContainerName, bRemoveNameFromSkeleton);
+		}
+		else
+		{
+			UE_LOG(LogAnimationBlueprintLibrary, Warning, TEXT("Could not find SmartName Container for Curve Name %s while trying to remove the curve"), *CurveName.ToString());
+		}
+	}
+	else
+	{
+		UE_LOG(LogAnimationBlueprintLibrary, Warning, TEXT("Invalid Animation Sequence for RemoveCurve"));
+	}
+}
+
+void UAnimationBlueprintLibrary::RemoveAllCurveData(UAnimSequence* AnimationSequence)
+{
+	if (AnimationSequence)
+	{
+		IAnimationDataController& Controller = AnimationSequence->GetController();
+
+		Controller.RemoveAllCurvesOfType(ERawCurveTrackTypes::RCT_Float);
+		Controller.RemoveAllCurvesOfType(ERawCurveTrackTypes::RCT_Transform);
+	}
+	else
+	{
+		UE_LOG(LogAnimationBlueprintLibrary, Warning, TEXT("Invalid Animation Sequence for RemoveAllCurveData"));
+	}
+}
+
+void UAnimationBlueprintLibrary::AddTransformationCurveKey(UAnimSequence* AnimationSequence, FName CurveName, const float Time, const FTransform& Transform)
+{
+	if (AnimationSequence)
+	{
+		TArray<float> TimeArray;
+		TArray<FTransform> TransformArray;
+
+		TimeArray.Add(Time);
+		TransformArray.Add(Transform);
+
+		AddCurveKeysInternal<FTransform, FTransformCurve, ERawCurveTrackTypes::RCT_Transform>(AnimationSequence, CurveName, TimeArray, TransformArray);
+	}
+	else
+	{
+		UE_LOG(LogAnimationBlueprintLibrary, Warning, TEXT("Invalid Animation Sequence for AddTransformationCurveKey"));
+	}
+
+}
+
+void UAnimationBlueprintLibrary::AddTransformationCurveKeys(UAnimSequence* AnimationSequence, FName CurveName, const TArray<float>& Times, const TArray<FTransform>& Transforms)
+{
+	if (AnimationSequence)
+	{
+		if (Times.Num() == Transforms.Num())
+		{
+			AddCurveKeysInternal<FTransform, FTransformCurve, ERawCurveTrackTypes::RCT_Transform>(AnimationSequence, CurveName, Times, Transforms);
+		}
+		else
+		{
+			UE_LOG(LogAnimationBlueprintLibrary, Warning, TEXT("Number of Time values %i does not match the number of Transforms %i in AddTransformationCurveKeys"), Times.Num(), Transforms.Num());
+		}	
+	}
+	else
+	{
+		UE_LOG(LogAnimationBlueprintLibrary, Warning, TEXT("Invalid Animation Sequence for AddTransformationCurveKeys"));
+	}
+}
+
+
+void UAnimationBlueprintLibrary::AddFloatCurveKey(UAnimSequence* AnimationSequence, FName CurveName, const float Time, const float Value)
+{
+	if (AnimationSequence)
+	{
+		TArray<float> TimeArray;
+		TArray<float> ValueArray;
+
+		TimeArray.Add(Time);
+		ValueArray.Add(Value);
+
+		AddCurveKeysInternal<float, FFloatCurve, ERawCurveTrackTypes::RCT_Float>(AnimationSequence, CurveName, TimeArray, ValueArray);
+	}
+	else
+	{
+		UE_LOG(LogAnimationBlueprintLibrary, Warning, TEXT("Invalid Animation Sequence for AddFloatCurveKey"));
+	}
+
+}
+
+void UAnimationBlueprintLibrary::AddFloatCurveKeys(UAnimSequence* AnimationSequence, FName CurveName, const TArray<float>& Times, const TArray<float>& Values)
+{
+	if (AnimationSequence)
+	{
+		if (Times.Num() == Values.Num())
+		{
+			AddCurveKeysInternal<float, FFloatCurve, ERawCurveTrackTypes::RCT_Float>(AnimationSequence, CurveName, Times, Values);
+		}
+		else
+		{
+			UE_LOG(LogAnimationBlueprintLibrary, Warning, TEXT("Number of Time values %i does not match the number of Values %i in AddFloatCurveKeys"), Times.Num(), Values.Num());			
+		}
+	}
+	else
+	{
+		UE_LOG(LogAnimationBlueprintLibrary, Warning, TEXT("Invalid Animation Sequence for AddFloatCurveKeys"));
+	}
+
+	
+}
+
+void UAnimationBlueprintLibrary::AddVectorCurveKey(UAnimSequence* AnimationSequence, FName CurveName, const float Time, const FVector Vector)
+{
+	if (AnimationSequence)
+	{
+		TArray<float> TimeArray;
+		TArray<FVector> VectorArray;
+
+		TimeArray.Add(Time);
+		VectorArray.Add(Vector);
+
+		AddCurveKeysInternal<FVector, FVectorCurve, ERawCurveTrackTypes::RCT_Vector>(AnimationSequence, CurveName, TimeArray, VectorArray);
+	}
+	else
+	{
+		UE_LOG(LogAnimationBlueprintLibrary, Warning, TEXT("Invalid Animation Sequence for AddVectorCurveKey"));
+	}
+
+}
+
+void UAnimationBlueprintLibrary::AddVectorCurveKeys(UAnimSequence* AnimationSequence, FName CurveName, const TArray<float>& Times, const TArray<FVector>& Vectors)
+{
+	if (AnimationSequence)
+	{
+		if (Times.Num() == Vectors.Num())
+		{
+			AddCurveKeysInternal<FVector, FVectorCurve, ERawCurveTrackTypes::RCT_Vector>(AnimationSequence, CurveName, Times, Vectors);
+		}
+		else
+		{
+			UE_LOG(LogAnimationBlueprintLibrary, Warning, TEXT("Number of Time values %i does not match the number of Vectors %i in AddVectorCurveKeys"), Times.Num(), Vectors.Num());
+		}
+	}
+	else
+	{
+		UE_LOG(LogAnimationBlueprintLibrary, Warning, TEXT("Invalid Animation Sequence for AddVectorCurveKeys"));
+	}
+}
+
+static void SetControllerCurveKeys(IAnimationDataController& Controller, FSmartName SmartName, const TArray<float>& Times, const TArray<float>& Values)
+{
+	const int32 NumKeys = Values.Num();
+	const FAnimationCurveIdentifier CurveId(SmartName, ERawCurveTrackTypes::RCT_Float);
+	for (int32 KeyIndex = 0; KeyIndex < NumKeys; ++KeyIndex)
+	{
+		Controller.SetCurveKey(CurveId, { Times[KeyIndex], Values[KeyIndex] });
+	}
+}
+
+static void SetControllerCurveKeys(IAnimationDataController& Controller, FSmartName SmartName, const TArray<float>& Times, const TArray<FTransform>& Values)
+{
+	const int32 NumKeys = Values.Num();
+	const FAnimationCurveIdentifier CurveId(SmartName, ERawCurveTrackTypes::RCT_Transform);
+	for (int32 KeyIndex = 0; KeyIndex < NumKeys; ++KeyIndex)
+	{
+		Controller.SetTransformCurveKey(CurveId, Times[KeyIndex], Values[KeyIndex]);
+	}
+}
+	
+static void SetControllerCurveKeys(IAnimationDataController& Controller, FSmartName CurveId, const TArray<float>& Times, const TArray<FVector>& Values)
+{
+	ensure(false);
+}
+
+template <typename DataType, typename CurveClass, ERawCurveTrackTypes CurveType>
+void UAnimationBlueprintLibrary::AddCurveKeysInternal(UAnimSequence* AnimationSequence, FName CurveName, const TArray<float>& Times, const TArray<DataType>& KeyData)
+{
+	checkf(Times.Num() == KeyData.Num(), TEXT("Not enough key data supplied"));
+
+	const FName ContainerName = RetrieveContainerNameForCurve(AnimationSequence, CurveName);
+
+	if (ContainerName != NAME_None)
+	{
+		// Retrieve smart name for curve
+		const FSmartName CurveSmartName = RetrieveSmartNameForCurve(AnimationSequence, CurveName, ContainerName);
+
+		// Retrieve the curve by name
+		const FAnimationCurveIdentifier CurveId(CurveSmartName, CurveType);
+		const CurveClass* Curve = static_cast<const CurveClass*>(AnimationSequence->GetDataModel()->FindCurve(CurveId));
+		if (Curve)
+		{
+			IAnimationDataController& Controller = AnimationSequence->GetController();
+			SetControllerCurveKeys(Controller, CurveSmartName, Times, KeyData);
+		}
+	}
+}
+
+bool UAnimationBlueprintLibrary::AddCurveInternal(UAnimSequence* AnimationSequence, FName CurveName, FName ContainerName, int32 CurveFlags, ERawCurveTrackTypes SupportedCurveType)
+{
+	// Add or retrieve the smart name
+	FSmartName SmartCurveName;
+	AnimationSequence->GetSkeleton()->AddSmartNameAndModify(ContainerName, CurveName, SmartCurveName);	
+
+	FAnimationCurveIdentifier CurveId(SmartCurveName, SupportedCurveType);
+
+	bool bCurveAdded = false;
+
+	const FAnimCurveBase* ExistingCurve = AnimationSequence->GetDataModel()->FindCurve(CurveId);
+	if (ExistingCurve == nullptr)
+	{
+		IAnimationDataController& Controller = AnimationSequence->GetController();
+		Controller.AddCurve(CurveId, CurveFlags);
+		bCurveAdded = AnimationSequence->GetDataModel()->FindCurve(CurveId) != nullptr;	
+	}
+	else
+	{
+		// Curve already exists
+	}
+
+	return bCurveAdded;
+}
+
+bool UAnimationBlueprintLibrary::RemoveCurveInternal(UAnimSequence* AnimationSequence, FName CurveName, FName ContainerName, bool bRemoveNameFromSkeleton)
+{
+	checkf(AnimationSequence != nullptr, TEXT("Invalid Animation Sequence ptr"));
+	bool bRemoved = false;
+	SmartName::UID_Type UID = AnimationSequence->GetSkeleton()->GetUIDByName(ContainerName, CurveName);
+	if (UID != SmartName::MaxUID)
+	{
+		FSmartName SmartCurveName;
+		USkeleton* Skeleton = AnimationSequence->GetSkeleton();
+		checkf(Skeleton != nullptr, TEXT("Invalid Skeleton ptr"));
+		if (Skeleton->GetSmartNameByUID(ContainerName, UID, SmartCurveName))
+		{
+			IAnimationDataController& Controller = AnimationSequence->GetController();
+
+			const FAnimationCurveIdentifier CurveId(SmartCurveName, ContainerName == USkeleton::AnimTrackCurveMappingName ? ERawCurveTrackTypes::RCT_Transform : ERawCurveTrackTypes::RCT_Float);
+			bRemoved = Controller.RemoveCurve(CurveId);
+
+			if (bRemoveNameFromSkeleton)
+			{
+				// Ensure we are eligible to do this
+				bool bValidToRemove = true;
+
+				if (ContainerName == USkeleton::AnimTrackCurveMappingName)
+				{
+					// Make sure we do not remove bone names
+					bValidToRemove = DoesBoneNameExistInternal(Skeleton, CurveName);
+				}
+
+				if (bValidToRemove)
+				{
+					Skeleton->RemoveSmartnameAndModify(ContainerName, UID);
+				}
+				else
+				{
+					UE_LOG(LogAnimationBlueprintLibrary, Warning, TEXT("Cannot remove Curve Name %s from Skeleton %s"), *CurveName.ToString(), *AnimationSequence->GetSkeleton()->GetName());
+				}
+			}
+		}
+		else
+		{
+			UE_LOG(LogAnimationBlueprintLibrary, Warning, TEXT("Could not retrieve Smart Name for Curve Name %s from Skeleton %s"), *CurveName.ToString(), *AnimationSequence->GetSkeleton()->GetName());
+		}
+	}
+	else
+	{
+		// Name does not exist on skeleton
+		UE_LOG(LogAnimationBlueprintLibrary, Warning, TEXT("Could find for Curve Name %s in Skeleton %s"), *CurveName.ToString(), *AnimationSequence->GetSkeleton()->GetName());
+	}
+
+	return bRemoved;
+}
+
+void UAnimationBlueprintLibrary::DoesBoneNameExist(UAnimSequence* AnimationSequence, FName BoneName, bool& bExists)
+{
+	bExists = false;
+	if (AnimationSequence)
+	{
+		USkeleton* Skeleton = AnimationSequence->GetSkeleton();
+		if (Skeleton)
+		{
+			bExists = DoesBoneNameExistInternal(Skeleton, BoneName);
+		}
+		else
+		{
+			UE_LOG(LogAnimationBlueprintLibrary, Warning, TEXT("No Skeleton found for Animation Sequence %s"), *AnimationSequence->GetName());
+		}
+	}
+	else
+	{
+		UE_LOG(LogAnimationBlueprintLibrary, Warning, TEXT("Invalid Animation Sequence supplied for DoesBoneNameExist"));
+	}
+}
+
+bool UAnimationBlueprintLibrary::DoesBoneNameExistInternal(USkeleton* Skeleton, FName BoneName)
+{
+	checkf(Skeleton != nullptr, TEXT("Invalid Skeleton ptr"));
+	return Skeleton->GetUIDByName(USkeleton::AnimTrackCurveMappingName, BoneName) != SmartName::MaxUID;
+}
+
+void UAnimationBlueprintLibrary::GetFloatKeys(UAnimSequence* AnimationSequence, FName CurveName, TArray<float>& Times, TArray<float>& Values)
+{
+	if (AnimationSequence)
+	{
+		GetCurveKeysInternal<float, FFloatCurve, ERawCurveTrackTypes::RCT_Float>(AnimationSequence, CurveName, Times, Values);
+	}
+	else
+	{
+		UE_LOG(LogAnimationBlueprintLibrary, Warning, TEXT("Invalid Animation Sequence for GetFloatKeys"));
+	}
+}
+
+void UAnimationBlueprintLibrary::GetVectorKeys(UAnimSequence* AnimationSequence, FName CurveName, TArray<float>& Times, TArray<FVector>& Values)
+{
+	if (AnimationSequence)
+	{
+		GetCurveKeysInternal<FVector, FVectorCurve, ERawCurveTrackTypes::RCT_Vector>(AnimationSequence, CurveName, Times, Values);
+	}
+	else
+	{
+		UE_LOG(LogAnimationBlueprintLibrary, Warning, TEXT("Invalid Animation Sequence for GetVectorKeys"));
+	}
+}
+
+void UAnimationBlueprintLibrary::GetTransformationKeys(UAnimSequence* AnimationSequence, FName CurveName, TArray<float>& Times, TArray<FTransform>& Values)
+{
+	if (AnimationSequence)
+	{
+		GetCurveKeysInternal<FTransform, FTransformCurve, ERawCurveTrackTypes::RCT_Transform>(AnimationSequence, CurveName, Times, Values);
+	}
+	else
+	{
+		UE_LOG(LogAnimationBlueprintLibrary, Warning, TEXT("Invalid Animation Sequence for GetTransformationKeys"));
+	}
+}
+
+void UAnimationBlueprintLibrary::CopyAnimationCurveNamesToSkeleton(USkeleton* OldSkeleton, USkeleton* NewSkeleton, UAnimSequenceBase* SequenceBase, ERawCurveTrackTypes CurveType)
+{
+	IAnimationDataController& Controller = SequenceBase->GetController();
+	Controller.UpdateCurveNamesFromSkeleton(OldSkeleton, CurveType);
+	Controller.FindOrAddCurveNamesOnSkeleton(NewSkeleton, CurveType);	
+}
+
+template <typename DataType, typename CurveClass, ERawCurveTrackTypes CurveType>
+void UAnimationBlueprintLibrary::GetCurveKeysInternal(UAnimSequence* AnimationSequence, FName CurveName, TArray<float>& Times, TArray<DataType>& KeyData)
+{
+	checkf(AnimationSequence != nullptr, TEXT("Invalid Animation Sequence ptr"));
+	const FName ContainerName = RetrieveContainerNameForCurve(AnimationSequence, CurveName);
+
+	if (ContainerName != NAME_None)
+	{
+		// Retrieve smart name for curve
+		const FSmartName CurveSmartName = RetrieveSmartNameForCurve(AnimationSequence, CurveName, ContainerName);
+
+		// Retrieve the curve by name
+		const FAnimationCurveIdentifier CurveId(CurveSmartName, CurveType);
+		const CurveClass* Curve = static_cast<const CurveClass*>(AnimationSequence->GetDataModel()->FindCurve(CurveId));
+
+		if (Curve)
+		{
+			Curve->GetKeys(Times, KeyData);
+			checkf(Times.Num() == KeyData.Num(), TEXT("Invalid key data retrieved from curve"));
+		}
+	}
+}
+
+bool UAnimationBlueprintLibrary::DoesCurveExist(UAnimSequence* AnimationSequence, FName CurveName, ERawCurveTrackTypes CurveType)
+{
+	bool bExistingCurve = false;
+
+	if (AnimationSequence)
+	{
+		FSmartName SmartName;
+		if (RetrieveSmartNameForCurve(AnimationSequence, CurveName, USkeleton::AnimTrackCurveMappingName, SmartName))
+		{
+			FAnimationCurveIdentifier CurveId(SmartName, CurveType);
+			const FAnimCurveBase* Curve = AnimationSequence->GetDataModel()->FindCurve(CurveId);
+			bExistingCurve = Curve != nullptr;
+		}
+
+		if (RetrieveSmartNameForCurve(AnimationSequence, CurveName, USkeleton::AnimCurveMappingName, SmartName))
+		{
+			FAnimationCurveIdentifier CurveId(SmartName, CurveType);
+			const FAnimCurveBase* Curve = AnimationSequence->GetDataModel()->FindCurve(CurveId);
+			bExistingCurve |= Curve != nullptr;
+		}
+	}
+	else
+	{
+		UE_LOG(LogAnimationBlueprintLibrary, Warning, TEXT("Invalid Animation Sequence for DoesCurveExist"));
+	}
+
+	return bExistingCurve;
+}
+
+bool UAnimationBlueprintLibrary::DoesSmartNameExist(UAnimSequence* AnimationSequence, FName Name)
+{
+	checkf(AnimationSequence != nullptr, TEXT("Invalid Animation Sequence ptr"));
+	FSmartName SmartName;
+	return AnimationSequence->GetSkeleton()->GetSmartNameByName(USkeleton::AnimTrackCurveMappingName, Name, SmartName) ||
+		AnimationSequence->GetSkeleton()->GetSmartNameByName(USkeleton::AnimCurveMappingName, Name, SmartName);
+}
+
+bool UAnimationBlueprintLibrary::RetrieveSmartNameForCurve(const UAnimSequence* AnimationSequence, FName CurveName, FName ContainerName, FSmartName& SmartName)
+{
+	checkf(AnimationSequence != nullptr, TEXT("Invalid Animation Sequence ptr"));
+	return AnimationSequence->GetSkeleton()->GetSmartNameByName(ContainerName, CurveName, SmartName);
+}
+
+FSmartName UAnimationBlueprintLibrary::RetrieveSmartNameForCurve(const UAnimSequence* AnimationSequence, FName CurveName, FName ContainerName)
+{
+	checkf(AnimationSequence != nullptr, TEXT("Invalid Animation Sequence ptr"));
+	FSmartName SmartCurveName;
+	AnimationSequence->GetSkeleton()->GetSmartNameByName(ContainerName, CurveName, SmartCurveName);
+	return SmartCurveName;
+}
+
+FName UAnimationBlueprintLibrary::RetrieveContainerNameForCurve(const UAnimSequence* AnimationSequence, FName CurveName)
+{
+	checkf(AnimationSequence != nullptr, TEXT("Invalid Animation Sequence ptr"));
+	const FName Names[(int32)ESmartNameContainerType::SNCT_MAX] = { USkeleton::AnimCurveMappingName, USkeleton::AnimTrackCurveMappingName };
+	for (int32 Index = 0; Index < (int32)ESmartNameContainerType::SNCT_MAX; ++Index)
+	{
+		const FSmartNameMapping* CurveMapping = AnimationSequence->GetSkeleton()->GetSmartNameContainer(Names[Index]);
+		if (CurveMapping && CurveMapping->Exists(CurveName))
+		{
+			return Names[Index];
+		}
+	}
+
+	return NAME_None;
+}
+
+void UAnimationBlueprintLibrary::AddMetaData(UAnimationAsset* AnimationAsset, TSubclassOf<UAnimMetaData> MetaDataClass, UAnimMetaData*& MetaDataInstance)
+{
+	if (AnimationAsset)
+	{
+		MetaDataInstance = NewObject<UAnimMetaData>(AnimationAsset, MetaDataClass, NAME_None, RF_Transactional);
+		if (MetaDataInstance)
+		{
+			AnimationAsset->AddMetaData(MetaDataInstance);
+		}
+		else
+		{
+			UE_LOG(LogAnimationBlueprintLibrary, Warning, TEXT("Failed to create instance for %s"), *MetaDataClass->GetName());
+		}
+		
+	}
+	else
+	{
+		UE_LOG(LogAnimationBlueprintLibrary, Warning, TEXT("Invalid Animation Sequence for AddMetaData"));
+	}
+}
+
+void UAnimationBlueprintLibrary::AddMetaDataObject(UAnimationAsset* AnimationAsset, UAnimMetaData* MetaDataObject)
+{
+	if (AnimationAsset && MetaDataObject)
+	{
+		if (MetaDataObject->GetOuter() == AnimationAsset)
+		{
+			AnimationAsset->AddMetaData(MetaDataObject);
+		}
+		else
+		{
+			UE_LOG(LogAnimationBlueprintLibrary, Warning, TEXT("Outer for MetaData Instance %s is not Animation Sequence %s"), *MetaDataObject->GetName(), *AnimationAsset->GetName());
+		}		
+	}
+	else 
+	{
+		if (!AnimationAsset)
+		{
+			UE_LOG(LogAnimationBlueprintLibrary, Warning, TEXT("Invalid Animation Sequence for AddMetaDataObject"));
+		}
+		
+		if (!MetaDataObject)
+		{
+			UE_LOG(LogAnimationBlueprintLibrary, Warning, TEXT("Invalid MetaDataObject for AddMetaDataObject"));
+		}
+	}
+}
+
+void UAnimationBlueprintLibrary::RemoveAllMetaData(UAnimationAsset* AnimationAsset)
+{
+	if (AnimationAsset)
+	{
+		AnimationAsset->EmptyMetaData();
+	}
+	else
+	{
+		UE_LOG(LogAnimationBlueprintLibrary, Warning, TEXT("Invalid Animation Sequence for RemoveAllMetaData"));
+	}
+}
+
+void UAnimationBlueprintLibrary::RemoveMetaData(UAnimationAsset* AnimationAsset, UAnimMetaData* MetaDataObject)
+{
+	if (AnimationAsset)
+	{
+		AnimationAsset->RemoveMetaData(MetaDataObject);
+	}
+	else
+	{
+		UE_LOG(LogAnimationBlueprintLibrary, Warning, TEXT("Invalid Animation Sequence for RemoveMetaData"));
+	}
+}
+
+void UAnimationBlueprintLibrary::RemoveMetaDataOfClass(UAnimationAsset* AnimationAsset, TSubclassOf<UAnimMetaData> MetaDataClass)
+{
+	if (AnimationAsset)
+	{
+		TArray<UAnimMetaData*> MetaDataOfClass;
+		GetMetaDataOfClass(AnimationAsset, MetaDataClass, MetaDataOfClass);
+		AnimationAsset->RemoveMetaData(MetaDataOfClass);
+	}
+	else
+	{
+		UE_LOG(LogAnimationBlueprintLibrary, Warning, TEXT("Invalid Animation Sequence for RemoveMetaDataOfClass"));
+	}	
+}
+
+void UAnimationBlueprintLibrary::GetMetaData(const UAnimationAsset* AnimationAsset, TArray<UAnimMetaData*>& MetaData)
+{
+	MetaData.Empty();
+	if (AnimationAsset)
+	{
+		MetaData = AnimationAsset->GetMetaData();
+	}
+	else
+	{
+		UE_LOG(LogAnimationBlueprintLibrary, Warning, TEXT("Invalid Animation Sequence for GetMetaData"));
+	}
+}
+
+void UAnimationBlueprintLibrary::GetMetaDataOfClass(const UAnimationAsset* AnimationAsset, TSubclassOf<UAnimMetaData> MetaDataClass, TArray<UAnimMetaData*>& MetaDataOfClass)
+{
+	MetaDataOfClass.Empty();
+	if (AnimationAsset)
+	{
+		const TArray<UAnimMetaData*>& MetaData = AnimationAsset->GetMetaData();
+		for (UAnimMetaData* MetaDataInstance : MetaData)
+		{
+			if (MetaDataInstance->GetClass() == *MetaDataClass)
+			{
+				MetaDataOfClass.Add(MetaDataInstance);
+			}
+		}
+	}
+	else
+	{
+		UE_LOG(LogAnimationBlueprintLibrary, Warning, TEXT("Invalid Animation Sequence for GetMetaDataOfClass"));
+	}
+}
+
+bool UAnimationBlueprintLibrary::ContainsMetaDataOfClass(const UAnimationAsset* AnimationAsset, TSubclassOf<UAnimMetaData> MetaDataClass)
+{
+	bool bContainsMetaData = false;
+	if (AnimationAsset)
+	{
+		TArray<UAnimMetaData*> MetaData;
+		GetMetaData(AnimationAsset, MetaData);
+		bContainsMetaData = MetaData.FindByPredicate(
+			[&](UAnimMetaData* MetaDataObject)
+		{
+			return (MetaDataObject->GetClass() == *MetaDataClass);
+		}) != nullptr;
+	}
+	else
+	{
+		UE_LOG(LogAnimationBlueprintLibrary, Warning, TEXT("Invalid Animation Sequence for ContainsMetaDataOfClass"));
+	}
+
+	return bContainsMetaData;	
+}
+
+void UAnimationBlueprintLibrary::GetBonePoseForTime(const UAnimSequenceBase* AnimationSequenceBase, FName BoneName, float Time, bool bExtractRootMotion, FTransform& Pose)
+{
+	Pose.SetIdentity();
+	if (AnimationSequenceBase)
+	{
+		TArray<FName> BoneNameArray;
+		TArray<FTransform> PoseArray;
+		BoneNameArray.Add(BoneName);
+		GetBonePosesForTime(AnimationSequenceBase, BoneNameArray, Time, bExtractRootMotion, PoseArray);
+		Pose = PoseArray[0];
+	}
+	else
+	{
+		UE_LOG(LogAnimationBlueprintLibrary, Warning, TEXT("Invalid Animation Sequence supplied for GetBonePoseForTime"));
+	}
+}
+
+void UAnimationBlueprintLibrary::GetBonePoseForFrame(const UAnimSequenceBase* AnimationSequenceBase, FName BoneName, int32 Frame, bool bExtractRootMotion, FTransform& Pose)
+{
+	Pose.SetIdentity();
+	if (AnimationSequenceBase)
+	{
+		GetBonePoseForTime(AnimationSequenceBase, BoneName, GetTimeAtFrameInternal(AnimationSequenceBase, Frame), bExtractRootMotion, Pose);
+	}
+	else
+	{
+		UE_LOG(LogAnimationBlueprintLibrary, Warning, TEXT("Invalid Animation Sequence supplied for GetBonePoseForFrame"));
+	}
+}
+
+void UAnimationBlueprintLibrary::GetBonePosesForTime(const UAnimSequenceBase* AnimationSequenceBase, TArray<FName> BoneNames, float Time, bool bExtractRootMotion, TArray<FTransform>& Poses, const USkeletalMesh* PreviewMesh /*= nullptr*/)
+{
+	Poses.Empty(BoneNames.Num());
+	if (AnimationSequenceBase && AnimationSequenceBase->GetSkeleton())
+	{
+		Poses.AddDefaulted(BoneNames.Num());
+
+		// Need this for FCompactPose
+		FMemMark Mark(FMemStack::Get());
+
+		const FReferenceSkeleton& RefSkeleton = (PreviewMesh)? PreviewMesh->GetRefSkeleton() : AnimationSequenceBase->GetSkeleton()->GetReferenceSkeleton();
+
+		if (IsValidTimeInternal(AnimationSequenceBase, Time))
+		{
+			if (BoneNames.Num())
+			{
+				for (int32 BoneNameIndex = 0; BoneNameIndex < BoneNames.Num(); ++BoneNameIndex)
+				{
+					const FName& BoneName = BoneNames[BoneNameIndex];
+					
+					FTransform& Transform = Poses[BoneNameIndex];
+					if (IsValidRawAnimationTrackName(AnimationSequenceBase, BoneName))
+					{
+						const EAnimInterpolationType InterpolationType = [AnimationSequenceBase]() -> EAnimInterpolationType
+					{
+							if (const UAnimSequence* AnimationSequence = Cast<const UAnimSequence>(AnimationSequenceBase))
+							{
+								return AnimationSequence->Interpolation;
+							}
+
+							return EAnimInterpolationType::Linear;
+						}();
+						UE::Anim::GetBoneTransformFromModel(AnimationSequenceBase->GetDataModel(), Transform, AnimationSequenceBase->GetDataModel()->GetBoneTrackIndexByName(BoneName), Time, InterpolationType);
+					}
+					else
+					{
+						
+						// otherwise, get ref pose if exists
+						const int32 BoneIndex = RefSkeleton.FindBoneIndex(BoneName);
+						if (BoneIndex != INDEX_NONE)
+						{
+							Transform = RefSkeleton.GetRefBonePose()[BoneIndex];
+						}
+						else
+						{
+							UE_LOG(LogAnimationBlueprintLibrary, Warning, TEXT("Invalid bone name %s for Animation Sequence %s supplied for GetBonePosesForTime"), *BoneName.ToString(), *AnimationSequenceBase->GetName());
+							Transform = FTransform::Identity;
+						}
+					}
+				}
+			}
+			else
+			{
+				UE_LOG(LogAnimationBlueprintLibrary, Error, TEXT("Invalid or no bone names specified to retrieve poses given Animation Sequence %s in GetBonePosesForTime"), *AnimationSequenceBase->GetName());
+			}			
+		}
+		else
+		{
+			UE_LOG(LogAnimationBlueprintLibrary, Warning, TEXT("Invalid time value %f for Animation Sequence %s supplied for GetBonePosesForTime"), Time, *AnimationSequenceBase->GetName());
+		}
+	}
+	else
+	{
+		UE_LOG(LogAnimationBlueprintLibrary, Warning, TEXT("Invalid Animation Sequence supplied for GetBonePosesForTime"));
+	}
+}
+
+void UAnimationBlueprintLibrary::GetBonePosesForFrame(const UAnimSequenceBase* AnimationSequenceBase, TArray<FName> BoneNames, int32 Frame, bool bExtractRootMotion, TArray<FTransform>& Poses, const USkeletalMesh* PreviewMesh /*= nullptr*/)
+{
+	Poses.Empty(BoneNames.Num());
+	if (AnimationSequenceBase)
+	{
+		GetBonePosesForTime(AnimationSequenceBase, BoneNames, GetTimeAtFrameInternal(AnimationSequenceBase, Frame), bExtractRootMotion, Poses, PreviewMesh);
+	}
+	else
+	{
+		UE_LOG(LogAnimationBlueprintLibrary, Warning, TEXT("Invalid Animation Sequence supplied for GetBonePosesForFrame"));
+	}
+}
+
+void UAnimationBlueprintLibrary::AddVirtualBone(const UAnimSequence* AnimationSequence, FName SourceBoneName, FName TargetBoneName, FName& VirtualBoneName)
+{
+	if (AnimationSequence)
+	{
+		USkeleton* Skeleton = AnimationSequence->GetSkeleton();
+		if (Skeleton)
+		{
+			const bool bSourceBoneExists = DoesBoneNameExistInternal(Skeleton, SourceBoneName);
+			const bool bTargetBoneExists = DoesBoneNameExistInternal(Skeleton, TargetBoneName);
+			const bool bVirtualBoneDoesNotExist = !DoesVirtualBoneNameExistInternal(Skeleton, VirtualBoneName);
+			
+			if (bSourceBoneExists && bTargetBoneExists && bVirtualBoneDoesNotExist)
+			{
+				const bool bAdded = Skeleton->AddNewVirtualBone(SourceBoneName, TargetBoneName, VirtualBoneName);
+				if (bAdded)
+				{
+					Skeleton->HandleSkeletonHierarchyChange();
+				}
+				else
+				{
+					UE_LOG(LogAnimationBlueprintLibrary, Warning, TEXT("Virtual bone between %s and %s already exists on Skeleton %s"), *SourceBoneName.ToString(), *TargetBoneName.ToString(), *Skeleton->GetName());
+				}
+			}
+			else
+			{
+				if (!bSourceBoneExists)
+				{
+					UE_LOG(LogAnimationBlueprintLibrary, Warning, TEXT("Bone Name %s does not exist on Skeleton %s"), *SourceBoneName.ToString(), *Skeleton->GetName());
+				}
+
+				if (!bTargetBoneExists)
+				{
+					UE_LOG(LogAnimationBlueprintLibrary, Warning, TEXT("Bone Name %s does not exist on Skeleton %s"), *TargetBoneName.ToString(), *Skeleton->GetName());
+				}			
+			}		
+		}
+		else
+		{
+			UE_LOG(LogAnimationBlueprintLibrary, Warning, TEXT("No Skeleton found for Animation Sequence %s"), *AnimationSequence->GetName());
+		}
+	}
+	else
+	{
+		UE_LOG(LogAnimationBlueprintLibrary, Warning, TEXT("Invalid Animation Sequence supplied for AddVirtualBone"));
+	}
+}
+
+void UAnimationBlueprintLibrary::RemoveVirtualBone(const UAnimSequence* AnimationSequence, FName VirtualBoneName)
+{
+	if (AnimationSequence)
+	{
+		USkeleton* Skeleton = AnimationSequence->GetSkeleton();
+		if (Skeleton)
+		{
+			if (DoesVirtualBoneNameExistInternal(Skeleton, VirtualBoneName))
+			{
+				TArray<FName> BoneNameArray;
+				BoneNameArray.Add(VirtualBoneName);
+				Skeleton->RemoveVirtualBones(BoneNameArray);
+				Skeleton->HandleSkeletonHierarchyChange();
+			}
+			else
+			{
+				UE_LOG(LogAnimationBlueprintLibrary, Warning, TEXT("Virtual Bone Name %s already exists on Skeleton %s"), *VirtualBoneName.ToString(), *Skeleton->GetName());
+			}
+		}
+		else
+		{
+			UE_LOG(LogAnimationBlueprintLibrary, Warning, TEXT("No Skeleton found for Animation Sequence %s"), *AnimationSequence->GetName());
+		}
+	}
+	else
+	{
+		UE_LOG(LogAnimationBlueprintLibrary, Warning, TEXT("Invalid Animation Sequence supplied for RemoveVirtualBone"));
+	}
+}
+
+void UAnimationBlueprintLibrary::RemoveVirtualBones(const UAnimSequence* AnimationSequence, TArray<FName> VirtualBoneNames)
+{
+	if (AnimationSequence)
+	{
+		USkeleton* Skeleton = AnimationSequence->GetSkeleton();
+		if (Skeleton)
+		{
+			for (FName& VirtualBoneName : VirtualBoneNames)
+			{
+				if (!DoesVirtualBoneNameExistInternal(Skeleton, VirtualBoneName))
+				{
+					UE_LOG(LogAnimationBlueprintLibrary, Warning, TEXT("Virtual Bone Name %s already exists on Skeleton %s"), *VirtualBoneName.ToString(), *Skeleton->GetName());
+				}
+			}			
+
+			Skeleton->RemoveVirtualBones(VirtualBoneNames);
+			Skeleton->HandleSkeletonHierarchyChange();
+		}
+		else
+		{
+			UE_LOG(LogAnimationBlueprintLibrary, Warning, TEXT("No Skeleton found for Animation Sequence %s"), *AnimationSequence->GetName());
+		}
+	}
+	else
+	{
+		UE_LOG(LogAnimationBlueprintLibrary, Warning, TEXT("Invalid Animation Sequence supplied for RemoveVirtualBones"));
+	}
+}
+
+void UAnimationBlueprintLibrary::RemoveAllVirtualBones(const UAnimSequence* AnimationSequence)
+{
+	if (AnimationSequence)
+	{
+		USkeleton* Skeleton = AnimationSequence->GetSkeleton();
+		if (Skeleton)
+		{
+			TArray<FName> VirtualBoneNames;
+			for (const FVirtualBone& VirtualBone : Skeleton->VirtualBones)
+			{
+				VirtualBoneNames.Add(VirtualBone.VirtualBoneName);
+			}
+
+			RemoveVirtualBones(AnimationSequence, VirtualBoneNames);
+		}
+		else
+		{
+			UE_LOG(LogAnimationBlueprintLibrary, Warning, TEXT("No Skeleton found for Animation Sequence %s"), *AnimationSequence->GetName());
+		}
+	}
+	else
+	{
+		UE_LOG(LogAnimationBlueprintLibrary, Warning, TEXT("Invalid Animation Sequence supplied for RemoveAllVirtualBones"));
+	}
+}
+
+bool UAnimationBlueprintLibrary::DoesVirtualBoneNameExistInternal(USkeleton* Skeleton, FName BoneName)
+{
+	checkf(Skeleton != nullptr, TEXT("Invalid Skeleton ptr"));
+	return Skeleton->VirtualBones.IndexOfByPredicate([&](const FVirtualBone& VirtualBone) { return VirtualBone.VirtualBoneName == BoneName; }) != INDEX_NONE;
+}
+
+void UAnimationBlueprintLibrary::GetSequenceLength(const UAnimSequenceBase* AnimationSequenceBase, float& Length)
+{
+	Length = 0.0f;
+	if (AnimationSequenceBase)
+	{
+		Length = AnimationSequenceBase->GetPlayLength();
+	}
+	else
+	{
+		UE_LOG(LogAnimationBlueprintLibrary, Warning, TEXT("Invalid Animation Sequence supplied for GetSequenceLength"));
+	}
+}
+
+void UAnimationBlueprintLibrary::GetRateScale(const UAnimSequenceBase* AnimationSequenceBase, float& RateScale)
+{
+	RateScale = 0.0f;
+	if (AnimationSequenceBase)
+	{
+		RateScale = AnimationSequenceBase->RateScale;
+	}
+	else
+	{
+		UE_LOG(LogAnimationBlueprintLibrary, Warning, TEXT("Invalid Animation Sequence supplied for GetRateScale"));
+	}
+}
+
+void UAnimationBlueprintLibrary::SetRateScale(UAnimSequenceBase* AnimationSequenceBase, float RateScale)
+{	
+	if (AnimationSequenceBase)
+	{
+		AnimationSequenceBase->RateScale = RateScale;
+	}
+	else
+	{
+		UE_LOG(LogAnimationBlueprintLibrary, Warning, TEXT("Invalid Animation Sequence supplied for SetRateScale"));
+	}
+}
+
+void UAnimationBlueprintLibrary::GetFrameAtTime(const UAnimSequenceBase* AnimationSequenceBase, const float Time, int32& Frame)
+{
+	Frame = 0;
+	if (AnimationSequenceBase)
+	{
+		Frame = AnimationSequenceBase->GetFrameAtTime(Time);
+	}
+	else
+	{		
+		UE_LOG(LogAnimationBlueprintLibrary, Warning, TEXT("Invalid Animation Sequence supplied for GetFrameAtTime"));
+	}
+}
+
+void UAnimationBlueprintLibrary::GetTimeAtFrame(const UAnimSequenceBase* AnimationSequenceBase, const int32 Frame, float& Time)
+{
+	Time = 0.0f;
+	if (AnimationSequenceBase)
+	{
+		Time = GetTimeAtFrameInternal(AnimationSequenceBase, Frame);
+	}
+	else
+	{		
+		UE_LOG(LogAnimationBlueprintLibrary, Warning, TEXT("Invalid Animation Sequence supplied for GetTimeAtFrame"));
+	}
+}
+
+float UAnimationBlueprintLibrary::GetTimeAtFrameInternal(const UAnimSequenceBase* AnimationSequenceBase, const int32 Frame)
+{
+	return AnimationSequenceBase->GetTimeAtFrame(Frame);
+}
+
+void UAnimationBlueprintLibrary::IsValidTime(const UAnimSequenceBase* AnimationSequenceBase, const float Time, bool& IsValid)
+{
+	IsValid = false;
+	if (AnimationSequenceBase)
+	{
+		IsValid = IsValidTimeInternal(AnimationSequenceBase, Time);
+	}
+	else
+	{		
+		UE_LOG(LogAnimationBlueprintLibrary, Warning, TEXT("Invalid Animation Sequence supplied for IsValidTime"));
+	}
+}
+
+bool UAnimationBlueprintLibrary::IsValidTimeInternal(const UAnimSequenceBase* AnimationSequenceBase, const float Time)
+{
+	return FMath::IsWithinInclusive(Time, 0.0f, AnimationSequenceBase->GetPlayLength());
+}
+
+void UAnimationBlueprintLibrary::FindBonePathToRoot(const UAnimSequenceBase* AnimationSequenceBase, FName BoneName, TArray<FName>& BonePath)
+{
+	BonePath.Empty();
+	if (AnimationSequenceBase)
+	{
+		BonePath.Add(BoneName);
+		int32 BoneIndex = AnimationSequenceBase->GetSkeleton()->GetReferenceSkeleton().FindRawBoneIndex(BoneName);		
+		if (BoneIndex != INDEX_NONE)
+		{
+			while (BoneIndex != INDEX_NONE)
+			{
+				const int32 ParentBoneIndex = AnimationSequenceBase->GetSkeleton()->GetReferenceSkeleton().GetRawParentIndex(BoneIndex);
+				if (ParentBoneIndex != INDEX_NONE)
+				{
+					BonePath.Add(AnimationSequenceBase->GetSkeleton()->GetReferenceSkeleton().GetBoneName(ParentBoneIndex));
+				}
+
+				BoneIndex = ParentBoneIndex;
+			}
+		}
+		else
+		{
+			UE_LOG(LogAnimationBlueprintLibrary, Warning, TEXT("Bone name %s not found in Skeleton %s"), *BoneName.ToString(), *AnimationSequenceBase->GetSkeleton()->GetName());
+		}
+	}
+	else
+	{
+		UE_LOG(LogAnimationBlueprintLibrary, Warning, TEXT("Invalid Animation Sequence supplied for FindBonePathToRoot"));
+	}
+}
+
+void UAnimationBlueprintLibrary::RemoveBoneAnimation(UAnimSequence* AnimationSequence, FName BoneName, bool bIncludeChildren /*= true*/, bool bFinalize /*= true*/)
+{
+	if (AnimationSequence)
+	{
+		TArray<FName> TrackNames;
+		AnimationSequence->GetDataModel()->GetBoneTrackNames(TrackNames);
+		
+		if (TrackNames.Contains(BoneName))
+		{
+			TArray<FName> TracksToRemove;
+			TracksToRemove.Add(BoneName);
+
+			// remove all children if required
+			if (bIncludeChildren)
+			{
+				USkeleton* Skeleton = AnimationSequence->GetSkeleton();
+				if (Skeleton)
+				{
+					const FReferenceSkeleton& RefSkeleton = Skeleton->GetReferenceSkeleton();
+					const int32 ParentBoneIndex = RefSkeleton.FindBoneIndex(BoneName);
+
+					// slow
+					for (const FName& TrackName : TrackNames)
+					{
+						if (TrackName != BoneName)
+						{
+							const int32 ChildBoneIndex = RefSkeleton.FindBoneIndex(TrackName);
+							if (RefSkeleton.BoneIsChildOf(ChildBoneIndex, ParentBoneIndex))
+							{
+								TracksToRemove.Add(TrackName);
+							}
+						}
+					}
+				}
+			}
+
+			IAnimationDataController& Controller = AnimationSequence->GetController();
+
+			IAnimationDataController::FScopedBracket ScopedBracket(Controller, LOCTEXT("RemoveBoneAnimation_Bracket", "Removing Bone Animation Track"));
+			for (const FName& TrackName : TracksToRemove)
+			{
+				Controller.RemoveBoneTrack(TrackName);
+			}
+		}
+		else
+		{
+			// print warning with track index
+			UE_LOG(LogAnimationBlueprintLibrary, Warning, TEXT("Invalid Bone Name for the animation."));
+		}
+	}
+}
+
+void UAnimationBlueprintLibrary::RemoveAllBoneAnimation(UAnimSequence* AnimationSequence)
+{
+	if (AnimationSequence)
+	{
+		IAnimationDataController& Controller = AnimationSequence->GetController();
+		IAnimationDataController::FScopedBracket ScopedBracket(Controller, LOCTEXT("RemoveAllBoneAnimation_Bracket", "Removing all Bone Animation and Transform Curve Tracks"));
+		Controller.RemoveAllBoneTracks();
+		Controller.RemoveAllCurvesOfType(ERawCurveTrackTypes::RCT_Transform);		
+	}
+}
+
+void UAnimationBlueprintLibrary::FinalizeBoneAnimation(UAnimSequence* AnimationSequence)
+{
+	if (AnimationSequence)
+	{
+		PRAGMA_DISABLE_DEPRECATION_WARNINGS
+		AnimationSequence->PostProcessSequence();
+		PRAGMA_ENABLE_DEPRECATION_WARNINGS
+	}
+}
+
+template void UAnimationBlueprintLibrary::AddCurveKeysInternal<float, FFloatCurve, ERawCurveTrackTypes::RCT_Float>(UAnimSequence* AnimationSequence, FName CurveName, const TArray<float>& Times, const TArray<float>& KeyData);
+template void UAnimationBlueprintLibrary::AddCurveKeysInternal<FVector, FVectorCurve, ERawCurveTrackTypes::RCT_Vector>(UAnimSequence* AnimationSequence, FName CurveName, const TArray<float>& Times, const TArray<FVector>& KeyData);
+template void UAnimationBlueprintLibrary::AddCurveKeysInternal<FTransform, FTransformCurve, ERawCurveTrackTypes::RCT_Transform>(UAnimSequence* AnimationSequence, FName CurveName, const TArray<float>& Times, const TArray<FTransform>& KeyData);
+
+template void UAnimationBlueprintLibrary::GetCurveKeysInternal<float, FFloatCurve, ERawCurveTrackTypes::RCT_Float>(UAnimSequence* AnimationSequence, FName CurveName, TArray<float>& Times, TArray<float>& KeyData);
+template void UAnimationBlueprintLibrary::GetCurveKeysInternal<FVector, FVectorCurve, ERawCurveTrackTypes::RCT_Vector>(UAnimSequence* AnimationSequence, FName CurveName, TArray<float>& Times, TArray<FVector>& KeyData);
+template void UAnimationBlueprintLibrary::GetCurveKeysInternal<FTransform, FTransformCurve, ERawCurveTrackTypes::RCT_Transform>(UAnimSequence* AnimationSequence, FName CurveName, TArray<float>& Times, TArray<FTransform>& KeyData);
+
+#undef LOCTEXT_NAMESPACE // "AnimationBlueprintLibrary"
\ No newline at end of file
diff --git a/Engine/Source/Editor/AnimationBlueprintLibrary/Public/AnimPose.h b/Engine/Source/Editor/AnimationBlueprintLibrary/Public/AnimPose.h
new file mode 100644
--- /dev/null
+++ b/Engine/Source/Editor/AnimationBlueprintLibrary/Public/AnimPose.h
@@ -0,0 +1,253 @@
+// Copyright Epic Games, Inc. All Rights Reserved.
+
+#pragma once
+
+#include "AnimationBlueprintLibrary.h"
+#include "Animation/AnimSequenceHelpers.h"
+#include "Kismet/BlueprintFunctionLibrary.h"
+
+#include "AnimPose.generated.h"
+
+/*	TODO
+ *	- Automated tests
+ *	- Verifying retargeted animation between different skeletons
+ */
+
+UENUM(BlueprintType)
+enum class EAnimPoseSpaces : uint8
+{
+	// Local (bone) space 
+	Local,
+	// World (component) space
+    World
+};
+
+UENUM(BlueprintType)
+enum class EAnimDataEvalType : uint8
+{
+	// Evaluates the original Animation Source data 
+	Source,
+	// Evaluates the original Animation Source data with additive animation layers
+	Raw,
+	// Evaluates the compressed Animation data - matching runtime (cooked)
+	Compressed
+};
+
+USTRUCT(BlueprintType)
+struct ANIMATIONBLUEPRINTLIBRARY_API FAnimPoseEvaluationOptions
+{
+	GENERATED_BODY()
+
+	// Type of evaluation which should be used
+	UPROPERTY(EditAnywhere, BlueprintReadWrite, Category="Animation|Pose")
+	EAnimDataEvalType EvaluationType = EAnimDataEvalType::Raw;
+
+	// Whether or not to retarget animation during evaluation
+	UPROPERTY(EditAnywhere, AdvancedDisplay, BlueprintReadWrite, Category="Animation|Pose")
+	bool bShouldRetarget = true;
+
+	// Whether or not to extract root motion values
+	UPROPERTY(EditAnywhere, AdvancedDisplay, BlueprintReadWrite, Category="Animation|Pose")
+	bool bExtractRootMotion = false;	
+};
+
+/** Script friendly representation of an evaluated animation bone pose */
+USTRUCT(BlueprintType)
+struct ANIMATIONBLUEPRINTLIBRARY_API FAnimPose
+{
+	GENERATED_BODY()
+
+	/** Returns whether or not the pose data was correctly initialized and populated */
+	bool IsValid() const;
+	
+protected:
+	/** Initializes the various arrays, using and copying the provided bone container */
+	void Init(const FBoneContainer& InBoneContainer);
+	/** Populates an FCompactPose using the contained bone data */
+	void GetPose(FCompactPose& InOutCompactPose) const;
+	/** Generates the contained bone data using the provided Component and its AnimInstance */
+	void SetPose(USkeletalMeshComponent* Component);
+	/** Generates the contained bone data using the provided CompactPose */
+	void SetPose(const FCompactPose& CompactPose);
+	/** Copies the reference pose to animated pose data */
+	void SetToRefPose();
+
+	/** (Re-)Generates the world space transforms using populated local space data */
+	void GenerateWorldSpaceTransforms();
+	/** Resets all contained data, rendering the instance invalid */
+	void Reset();
+
+	/** Whether or not the contained data was initialized and can be used to store a pose */
+	bool IsInitialized() const { return BoneNames.Num() != 0; }
+	/** Whether or local space pose data has been populated */
+	bool IsPopulated() const { return LocalSpacePoses.Num() != 0; }
+	
+protected:
+	UPROPERTY()
+	TArray<FName> BoneNames;
+	
+	UPROPERTY()
+	TArray<int32> BoneIndices;
+
+	UPROPERTY()
+	TArray<int32> ParentBoneIndices;
+
+	UPROPERTY()
+	TArray<FTransform> LocalSpacePoses;
+	
+	UPROPERTY()
+	TArray<FTransform> WorldSpacePoses;
+
+	UPROPERTY()
+	TArray<FTransform> RefLocalSpacePoses;
+	
+	UPROPERTY()
+	TArray<FTransform> RefWorldSpacePoses;
+
+	FBoneContainer BoneContainer;
+
+	friend class UAnimPoseExtensions;
+};
+
+/** Script exposed functionality for populating, retrieving data from and setting data on FAnimPose */
+UCLASS()
+class ANIMATIONBLUEPRINTLIBRARY_API UAnimPoseExtensions : public UBlueprintFunctionLibrary
+{
+	GENERATED_BODY()
+	
+	/**
+	*  Returns whether the Anim Pose contains valid data
+	*
+	* @param	Pose	Anim Pose to validate
+	*
+	* @return	Result of the validation
+	*/	
+	UFUNCTION(BlueprintPure, meta = (ScriptMethod), Category = "Animation|Pose")
+	static bool IsValid(UPARAM(ref) const FAnimPose& Pose);
+
+	/** */
+	/**
+	* Returns an array of bone names contained by the pose
+	*
+	* @param	Pose	Anim Pose to retrieve the names from
+	* @param	Bones	Array to be populated with the bone names
+	*
+	*/
+	UFUNCTION(BlueprintPure, meta = (ScriptMethod), Category = "Animation|Pose")
+	static void GetBoneNames(UPARAM(ref) const FAnimPose& Pose, TArray<FName>& Bones);
+
+	/**
+	* Retrieves the transform for the provided bone name from a pose
+	*
+	* @param	Pose		Anim Pose to retrieve the transform from
+	* @param	BoneName	Name of the bone to retrieve
+	* @param	Space		Space in which the transform should be retrieved
+	*
+	* @return	Transform in requested space for bone if found, otherwise return identity transform
+	*/
+	UFUNCTION(BlueprintPure, meta = (ScriptMethod), Category = "Animation|Pose")
+    static FTransform GetBonePose(UPARAM(ref) const FAnimPose& Pose, FName BoneName, EAnimPoseSpaces Space = EAnimPoseSpaces::Local);
+
+	/**
+	* Sets the transform for the provided bone name for a pose
+	*
+	* @param	Pose		Anim Pose to set transform in
+	* @param	Transform	Transform to set the bone to
+	* @param	BoneName	Name of the bone to set
+	* @param	Space		Space in which the transform should be set
+	*/
+	UFUNCTION(BlueprintCallable, meta = (ScriptMethod), Category = "Animation|Pose")
+	static void SetBonePose(UPARAM(ref) FAnimPose& Pose, FTransform Transform, FName BoneName, EAnimPoseSpaces Space = EAnimPoseSpaces::Local);	
+
+	/**
+	* Retrieves the reference pose transform for the provided bone name
+	*
+	* @param	Pose		Anim Pose to retrieve the transform from
+	* @param	BoneName	Name of the bone to retrieve
+	* @param	Space		Space in which the transform should be retrieved
+	*
+	* @return	Transform in requested space for bone if found, otherwise return identity transform
+	*/
+	UFUNCTION(BlueprintPure, meta = (ScriptMethod), Category = "Animation|Pose")
+    static FTransform GetRefBonePose(UPARAM(ref) const FAnimPose& Pose, FName BoneName, EAnimPoseSpaces Space = EAnimPoseSpaces::Local);
+
+	/**
+	* Retrieves the relative transform between the two provided bone names
+	*
+	* @param	Pose			Anim Pose to retrieve the transform from
+	* @param	FromBoneName	Name of the bone to retrieve the transform relative from
+	* @param	ToBoneName		Name of the bone to retrieve the transform relative to
+	* @param	Space			Space in which the transform should be retrieved
+	*
+	* @return	Relative transform in requested space for bone if found, otherwise return identity transform
+	*/
+	UFUNCTION(BlueprintPure, meta = (ScriptMethod), Category = "Animation|Pose")
+    static FTransform GetRelativeTransform(UPARAM(ref) const FAnimPose& Pose, FName FromBoneName, FName ToBoneName, EAnimPoseSpaces Space = EAnimPoseSpaces::Local);
+
+	/**
+	* Retrieves the relative transform between reference and animated bone transform
+	*
+	* @param	Pose			Anim Pose to retrieve the transform from
+	* @param	BoneName		Name of the bone to retrieve the relative transform for
+	* @param	Space			Space in which the transform should be retrieved
+	*
+	* @return	Relative transform in requested space for bone if found, otherwise return identity transform
+	*/
+	UFUNCTION(BlueprintPure, meta = (ScriptMethod), Category = "Animation|Pose")
+    static FTransform GetRelativeToRefPoseTransform(UPARAM(ref) const FAnimPose& Pose, FName BoneName, EAnimPoseSpaces Space = EAnimPoseSpaces::Local);
+
+	/**
+	* Retrieves the relative transform for the reference pose between the two provided bone names 
+	*
+	* @param	Pose			Anim Pose to retrieve the transform from
+	* @param	FromBoneName	Name of the bone to retrieve the transform relative from
+	* @param	ToBoneName		Name of the bone to retrieve the transform relative to
+	* @param	Space			Space in which the transform should be retrieved
+	*
+	* @return	Relative transform in requested space for bone if found, otherwise return identity transform
+	*/
+	UFUNCTION(BlueprintPure, meta = (ScriptMethod), Category = "Animation|Pose")
+    static FTransform GetRefPoseRelativeTransform(UPARAM(ref) const FAnimPose& Pose, FName FromBoneName, FName ToBoneName, EAnimPoseSpaces Space = EAnimPoseSpaces::Local);
+
+	/**
+	* Evaluates an Animation Sequence Base to generate a valid Anim Pose instance
+	*
+	* @param	AnimationSequenceBase	Animation sequence base to evaluate the pose from
+	* @param	Time					Time at which the pose should be evaluated
+	* @param	EvaluationOptions		Options determining the way the pose should be evaluated
+	* @param	Pose					Anim Pose to hold the evaluated data
+	*/
+	UFUNCTION(BlueprintCallable, meta = (ScriptMethod), Category = "Animation|Pose")
+	static void GetAnimPoseAtTime(const UAnimSequenceBase* AnimationSequenceBase, float Time, FAnimPoseEvaluationOptions EvaluationOptions, FAnimPose& Pose);
+
+	/**
+	* Evaluates an Animation Sequence Base to generate a valid Anim Pose instance
+	*
+	* @param	AnimationSequenceBase	Animation sequence base to evaluate the pose from
+	* @param	FrameIndex				Exact frame at which the pose should be evaluated
+	* @param	EvaluationOptions		Options determining the way the pose should be evaluated
+	* @param	Pose					Anim Pose to hold the evaluated data
+	*/
+	UFUNCTION(BlueprintCallable, meta = (ScriptMethod), Category = "Animation|Pose")
+	static void GetAnimPoseAtFrame(const UAnimSequenceBase* AnimationSequenceBase, int32 FrameIndex, FAnimPoseEvaluationOptions EvaluationOptions, FAnimPose& Pose);
+
+	/**
+	* Evaluates an Animation Blueprint instance, using the provided Anim Pose and its Input Pose value, generating a valid Anim Pose using the result. Warning this function may cause performance issues.
+	*
+	* @param	InputPose				Anim Pose used to populate the input pose node value inside of the Animation Blueprint
+	* @param	TargetSkeletalMesh		USkeletalMesh object used as the target skeletal mesh, should have same USkeleton as InputPose and Animation Blueprint
+	* @param	AnimationBlueprint		Animation Blueprint to generate an AnimInstance with, used to evaluate the output Anim Pose
+	* @param	OutPose					Anim pose to hold the data from evaluating the Animation Blueprint instance		
+	*/
+	UFUNCTION(BlueprintCallable, meta = (ScriptMethod), Category = "Animation|Pose")
+	static void EvaluateAnimationBlueprintWithInputPose(const FAnimPose& InputPose, USkeletalMesh* TargetSkeletalMesh, UAnimBlueprint* AnimationBlueprint, FAnimPose& OutPose);
+
+	/**
+	* Populates an Anim Pose with the reference poses stored for the provided USkeleton
+	*
+	* @param	Skeleton				USkeleton object to retrieve the reference pose from
+	* @param	OutPose					Anim pose to hold the reference pose
+	*/
+	UFUNCTION(BlueprintPure, meta = (ScriptMethod), Category = "Animation|Pose")
+	static void GetReferencePose(USkeleton* Skeleton, FAnimPose& OutPose);
+};
\ No newline at end of file
diff --git a/Engine/Source/Editor/AnimationBlueprintLibrary/Public/AnimationBlueprintLibrary.h b/Engine/Source/Editor/AnimationBlueprintLibrary/Public/AnimationBlueprintLibrary.h
new file mode 100644
--- /dev/null
+++ b/Engine/Source/Editor/AnimationBlueprintLibrary/Public/AnimationBlueprintLibrary.h
@@ -0,0 +1,473 @@
+// Copyright Epic Games, Inc. All Rights Reserved.
+
+#pragma once
+
+#include "Kismet/BlueprintFunctionLibrary.h"
+#include "Animation/AnimEnums.h"
+#include "Animation/AnimCurveTypes.h"
+#include "Animation/AnimMetaData.h"
+
+#include "AnimationBlueprintLibrary.generated.h"
+
+struct FRawAnimSequenceTrack;
+class UAnimCompress;
+class USkeleton;
+class UAnimBoneCompressionSettings;
+class UAnimCurveCompressionSettings;
+
+UENUM()
+enum class ESmartNameContainerType : uint8
+{
+	SNCT_CurveMapping UMETA(DisplayName = "Curve Names"),
+	SNCT_TrackCurveMapping	UMETA(DisplayName = "Track Curve Names"),
+	SNCT_MAX
+};
+
+/** Delegate called when a notify was replaced */
+DECLARE_DYNAMIC_DELEGATE_TwoParams(FOnNotifyReplaced, UAnimNotify*, OldNotify, UAnimNotify*, NewNotify);
+
+/** Delegate called when a notify state was replaced */
+DECLARE_DYNAMIC_DELEGATE_TwoParams(FOnNotifyStateReplaced, UAnimNotifyState*, OldNotifyState, UAnimNotifyState*, NewNotifyState);
+
+/** Blueprint library for altering and analyzing animation / skeletal data */
+UCLASS(meta=(ScriptName="AnimationLibrary"))
+class ANIMATIONBLUEPRINTLIBRARY_API UAnimationBlueprintLibrary : public UBlueprintFunctionLibrary
+{
+	GENERATED_BODY()
+
+public:
+	/** Retrieves the number of animation frames for the given Animation Sequence */
+	UFUNCTION(BlueprintPure, meta=(AutoCreateRefTerm = "AnimationSequence"), Category = "AnimationBlueprintLibrary|Animation")
+	static void GetNumFrames(const UAnimSequenceBase* AnimationSequenceBase, int32& NumFrames);
+
+	/** Retrieves the number of animation keys for the given Animation Sequence */
+	UFUNCTION(BlueprintPure, meta = (AutoCreateRefTerm = "AnimationSequence"), Category = "AnimationBlueprintLibrary|Animation")
+	static void GetNumKeys(const UAnimSequenceBase* AnimationSequenceBase, int32& NumKeys);
+	
+	/** Retrieves the Names of the individual ATracks for the given Animation Sequence */
+	UFUNCTION(BlueprintPure, Category = "AnimationBlueprintLibrary|Animation")
+	static void GetAnimationTrackNames(const UAnimSequenceBase* AnimationSequenceBase, TArray<FName>& TrackNames);
+
+	/** Retrieves the Names of the individual float curves for the given Animation Sequence */
+	UFUNCTION(BlueprintPure, Category = "AnimationBlueprintLibrary|Animation")
+	static void GetAnimationCurveNames(const UAnimSequence* AnimationSequence, ERawCurveTrackTypes CurveType, TArray<FName>& CurveNames);
+
+	/** Retrieves the Raw Translation Animation Data for the given Animation Track Name and Animation Sequence */
+	UFUNCTION(BlueprintPure, Category = "AnimationBlueprintLibrary|RawTrackData")
+	static void GetRawTrackPositionData(const UAnimSequenceBase* AnimationSequenceBase, const FName TrackName, TArray<FVector>& PositionData);
+
+	/** Retrieves the Raw Rotation Animation Data for the given Animation Track Name and Animation Sequence */
+	UFUNCTION(BlueprintPure, Category = "AnimationBlueprintLibrary|RawTrackData")
+	static void GetRawTrackRotationData(const UAnimSequenceBase* AnimationSequenceBase, const FName TrackName, TArray<FQuat>& RotationData );
+
+	/** Retrieves the Raw Scale Animation Data for the given Animation Track Name and Animation Sequence */
+	UFUNCTION(BlueprintPure, Category = "AnimationBlueprintLibrary|RawTrackData")
+	static void GetRawTrackScaleData(const UAnimSequenceBase* AnimationSequenceBase, const FName TrackName, TArray<FVector>& ScaleData);
+
+	/** Retrieves the Raw Animation Data for the given Animation Track Name and Animation Sequence */
+	UFUNCTION(BlueprintPure, Category = "AnimationBlueprintLibrary|RawTrackData")
+	static void GetRawTrackData(const UAnimSequenceBase* AnimationSequenceBase, const FName TrackName, TArray<FVector>& PositionKeys,TArray<FQuat>& RotationKeys, TArray<FVector>& ScalingKeys);
+
+	/** Checks whether or not the given Animation Track Name is contained within the Animation Sequence */
+	UFUNCTION(BlueprintCallable, Category = "AnimationBlueprintLibrary|Helpers")
+	static bool IsValidRawAnimationTrackName(const UAnimSequenceBase* AnimationSequenceBase, const FName TrackName);
+
+	static const FRawAnimSequenceTrack& GetRawAnimationTrackByName(const UAnimSequenceBase* AnimationSequenceBase, const FName TrackName);
+
+	// Compression
+
+	/** Retrieves the Bone Compression Settings for the given Animation Sequence */
+	UFUNCTION(BlueprintPure, Category = "AnimationBlueprintLibrary|Compression")
+	static void GetBoneCompressionSettings(const UAnimSequence* AnimationSequence, UAnimBoneCompressionSettings*& CompressionSettings);
+
+	/** Sets the Bone Compression Settings for the given Animation Sequence */
+	UFUNCTION(BlueprintCallable, Category = "AnimationBlueprintLibrary|Compression")
+	static void SetBoneCompressionSettings(UAnimSequence* AnimationSequence, UAnimBoneCompressionSettings* CompressionSettings);
+
+	/** Retrieves the Curve Compression Settings for the given Animation Sequence */
+	UFUNCTION(BlueprintPure, Category = "AnimationBlueprintLibrary|Compression")
+	static void GetCurveCompressionSettings(const UAnimSequence* AnimationSequence, UAnimCurveCompressionSettings*& CompressionSettings);
+
+	/** Sets the Curve Compression Settings for the given Animation Sequence */
+	UFUNCTION(BlueprintCallable, Category = "AnimationBlueprintLibrary|Compression")
+	static void SetCurveCompressionSettings(UAnimSequence* AnimationSequence, UAnimCurveCompressionSettings* CompressionSettings);
+
+	// Additive 
+	/** Retrieves the Additive Animation type for the given Animation Sequence */
+	UFUNCTION(BlueprintPure, Category = "AnimationBlueprintLibrary|Additive")
+	static void GetAdditiveAnimationType(const UAnimSequence* AnimationSequence, TEnumAsByte<enum EAdditiveAnimationType>& AdditiveAnimationType);
+
+	/** Sets the Additive Animation type for the given Animation Sequence */
+	UFUNCTION(BlueprintCallable, Category = "AnimationBlueprintLibrary|Additive")
+	static void SetAdditiveAnimationType(UAnimSequence* AnimationSequence, const TEnumAsByte<enum EAdditiveAnimationType> AdditiveAnimationType);
+
+	/** Retrieves the Additive Base Pose type for the given Animation Sequence */
+	UFUNCTION(BlueprintPure, Category = "AnimationBlueprintLibrary|Additive")
+	static void GetAdditiveBasePoseType(const UAnimSequence* AnimationSequence, TEnumAsByte<enum EAdditiveBasePoseType>& AdditiveBasePoseType);
+
+	/** Sets the Additive Base Pose type for the given Animation Sequence */
+	UFUNCTION(BlueprintCallable, Category = "AnimationBlueprintLibrary|Additive")
+	static void SetAdditiveBasePoseType(UAnimSequence* AnimationSequence, const TEnumAsByte<enum EAdditiveBasePoseType> AdditiveBasePoseType);
+
+	// Interpolation
+
+	/** Retrieves the Animation Interpolation type for the given Animation Sequence */
+	UFUNCTION(BlueprintPure, Category = "AnimationBlueprintLibrary|Interpolation")
+	static void GetAnimationInterpolationType(const UAnimSequence* AnimationSequence, EAnimInterpolationType& InterpolationType);
+
+	/** Sets the Animation Interpolation type for the given Animation Sequence */
+	UFUNCTION(BlueprintCallable, Category = "AnimationBlueprintLibrary|Interpolation")
+	static void SetAnimationInterpolationType(UAnimSequence* AnimationSequence, EAnimInterpolationType InterpolationType);
+
+	// Root motion
+
+	/** Checks whether or not Root Motion is Enabled for the given Animation Sequence */
+	UFUNCTION(BlueprintPure, Category = "AnimationBlueprintLibrary|RootMotion")
+	static bool IsRootMotionEnabled(const UAnimSequence* AnimationSequence);
+
+	/** Sets whether or not Root Motion is Enabled for the given Animation Sequence */
+	UFUNCTION(BlueprintCallable, Category = "AnimationBlueprintLibrary|RootMotion")
+	static void SetRootMotionEnabled(UAnimSequence* AnimationSequence, bool bEnabled);
+
+	/** Retrieves the Root Motion Lock Type for the given Animation Sequence */
+	UFUNCTION(BlueprintPure, Category = "AnimationBlueprintLibrary|RootMotion")
+	static void GetRootMotionLockType(const UAnimSequence* AnimationSequence, TEnumAsByte<ERootMotionRootLock::Type>& LockType);
+
+	/** Sets the Root Motion Lock Type for the given Animation Sequence */
+	UFUNCTION(BlueprintCallable, Category = "AnimationBlueprintLibrary|RootMotion")
+	static void SetRootMotionLockType(UAnimSequence* AnimationSequence, TEnumAsByte<ERootMotionRootLock::Type> RootMotionLockType);
+
+	/** Checks whether or not Root Motion locking is Forced for the given Animation Sequence */
+	UFUNCTION(BlueprintPure, Category = "AnimationBlueprintLibrary|RootMotion")
+	static bool IsRootMotionLockForced(const UAnimSequence* AnimationSequence);
+
+	/** Sets whether or not Root Motion locking is Forced for the given Animation Sequence */
+	UFUNCTION(BlueprintCallable, Category = "AnimationBlueprintLibrary|RootMotion")
+	static void SetIsRootMotionLockForced(UAnimSequence* AnimationSequence, bool bForced);
+
+	// Markers
+
+	/** Retrieves all the Animation Sync Markers for the given Animation Sequence */
+	UFUNCTION(BlueprintPure, Category = "AnimationBlueprintLibrary|MarkerSyncing")
+	static void GetAnimationSyncMarkers(const UAnimSequence* AnimationSequence, TArray<FAnimSyncMarker>& Markers);
+
+	/** Retrieves all the Unique Names for the Animation Sync Markers contained by the given Animation Sequence */
+	UFUNCTION(BlueprintPure, Category = "AnimationBlueprintLibrary|MarkerSyncing")
+	static void GetUniqueMarkerNames(const UAnimSequence* AnimationSequence, TArray<FName>& MarkerNames);
+
+	/** Adds an Animation Sync Marker to Notify track in the given Animation with the corresponding Marker Name and Time */
+	UFUNCTION(BlueprintCallable, Category = "AnimationBlueprintLibrary|MarkerSyncing")
+	static void AddAnimationSyncMarker(UAnimSequence* AnimationSequence, FName MarkerName, float Time, FName NotifyTrackName);
+
+	/** Checks whether or not the given Marker Name is a valid Animation Sync Marker Name */
+	UFUNCTION(BlueprintPure, Category = "AnimationBlueprintLibrary|Helpers")
+	static bool IsValidAnimationSyncMarkerName(const UAnimSequence* AnimationSequence, FName MarkerName);
+
+	/** Removes All Animation Sync Marker found within the Animation Sequence whose name matches MarkerName, and returns the number of removed instances */
+	UFUNCTION(BlueprintCallable, Category = "AnimationBlueprintLibrary|MarkerSyncing")
+	static int32 RemoveAnimationSyncMarkersByName(UAnimSequence* AnimationSequence, FName MarkerName);
+
+	/** Removes All Animation Sync Marker found within the Animation Sequence that belong to the specific Notify Track, and returns the number of removed instances */	
+	UFUNCTION(BlueprintCallable, Category = "AnimationBlueprintLibrary|MarkerSyncing")
+	static int32 RemoveAnimationSyncMarkersByTrack(UAnimSequence* AnimationSequence, FName NotifyTrackName);
+
+	/** Removes All Animation Sync Markers found within the Animation Sequence, and returns the number of removed instances */
+	UFUNCTION(BlueprintCallable, Category = "AnimationBlueprintLibrary|MarkerSyncing")
+	static void RemoveAllAnimationSyncMarkers(UAnimSequence* AnimationSequence);
+
+	// Notifies
+
+	/** Retrieves all Animation Notify Events found within the given Animation Sequence */
+	UFUNCTION(BlueprintPure, Category = "AnimationBlueprintLibrary|NotifyEvents")
+	static void GetAnimationNotifyEvents(const UAnimSequenceBase* AnimationSequenceBase, TArray<FAnimNotifyEvent>& NotifyEvents);
+
+	/** Retrieves all Unique Animation Notify Events found within the given Animation Sequence */
+	UFUNCTION(BlueprintPure, Category = "AnimationBlueprintLibrary|NotifyEvents")
+	static void GetAnimationNotifyEventNames(const UAnimSequenceBase* AnimationSequenceBase, TArray<FName>& EventNames);
+
+	/** Adds an Animation Notify Event to Notify track in the given Animation with the given Notify creation data */
+	UFUNCTION(BlueprintCallable, Category = "AnimationBlueprintLibrary|NotifyEvents")
+	static UAnimNotify* AddAnimationNotifyEvent(UAnimSequenceBase* AnimationSequenceBase, FName NotifyTrackName, float StartTime, TSubclassOf<UAnimNotify> NotifyClass);
+
+	/** Adds an Animation Notify State Event to Notify track in the given Animation with the given Notify State creation data */
+	UFUNCTION(BlueprintCallable, Category = "AnimationBlueprintLibrary|NotifyEvents")
+	static UAnimNotifyState* AddAnimationNotifyStateEvent(UAnimSequenceBase* AnimationSequenceBase, FName NotifyTrackName, float StartTime, float Duration, TSubclassOf<UAnimNotifyState> NotifyStateClass);
+
+	/** Adds an the specific Animation Notify to the Animation Sequence (requires Notify's outer to be the Animation Sequence) */
+	UFUNCTION(BlueprintCallable, Category = "AnimationBlueprintLibrary|NotifyEvents")
+	static void AddAnimationNotifyEventObject(UAnimSequenceBase* AnimationSequenceBase, float StartTime, UAnimNotify* Notify, FName NotifyTrackName);
+
+	/** Adds an the specific Animation Notify State to the Animation Sequence (requires Notify State's outer to be the Animation Sequence) */
+	UFUNCTION(BlueprintCallable, Category = "AnimationBlueprintLibrary|NotifyEvents")
+	static void AddAnimationNotifyStateEventObject(UAnimSequenceBase* AnimationSequenceBase, float StartTime, float Duration, UAnimNotifyState* NotifyState, FName NotifyTrackName);
+
+	/** Removes Animation Notify Events found by Name within the Animation Sequence, and returns the number of removed name instances */
+	UFUNCTION(BlueprintCallable, Category = "AnimationBlueprintLibrary|NotifyEvents")
+	static int32 RemoveAnimationNotifyEventsByName(UAnimSequenceBase* AnimationSequenceBase, FName NotifyName);
+
+	/** Removes Animation Notify Events found by Track within the Animation Sequence, and returns the number of removed name instances */
+	UFUNCTION(BlueprintCallable, Category = "AnimationBlueprintLibrary|NotifyEvents")
+	static int32 RemoveAnimationNotifyEventsByTrack(UAnimSequenceBase* AnimationSequenceBase, FName NotifyTrackName);	
+
+	/** Replaces animation notifies in the specified Animation Sequence */
+	UFUNCTION(BlueprintCallable, Category = "AnimationBlueprintLibrary|NotifyEvents")
+	static void ReplaceAnimNotifyStates(UAnimSequenceBase* AnimationSequenceBase, TSubclassOf<UAnimNotifyState> OldNotifyClass, TSubclassOf<UAnimNotifyState> NewNotifyClass, FOnNotifyStateReplaced OnNotifyStateReplaced);
+
+	/** Replaces animation notifies in the specified Animation Sequence */
+	UFUNCTION(BlueprintCallable, Category = "AnimationBlueprintLibrary|NotifyEvents")
+	static void ReplaceAnimNotifies(UAnimSequenceBase* AnimationSequenceBase, TSubclassOf<UAnimNotify> OldNotifyClass, TSubclassOf<UAnimNotify> NewNotifyClass, FOnNotifyReplaced OnNotifyReplaced);
+
+	/** Copies animation notifies from Src Animation Sequence to Dest. Creates anim notify tracks as necessary. Returns true on success. */
+	UFUNCTION(BlueprintCallable, Category = "AnimationBlueprintLibrary|NotifyEvents")
+	static void CopyAnimNotifiesFromSequence(UAnimSequenceBase* SourceAnimationSequenceBase, UAnimSequenceBase* DestinationAnimationSequenceBase);
+
+	// Notify Tracks
+
+	/** Retrieves all Unique Animation Notify Track Names found within the given Animation Sequence */
+	UFUNCTION(BlueprintPure, Category = "AnimationBlueprintLibrary|AnimationNotifies")
+	static void GetAnimationNotifyTrackNames(const UAnimSequenceBase* AnimationSequenceBase, TArray<FName>& TrackNames);
+
+	/** Adds an Animation Notify Track to the Animation Sequence */
+	UFUNCTION(BlueprintCallable, Category = "AnimationBlueprintLibrary|AnimationNotifies")
+	static void AddAnimationNotifyTrack(UAnimSequenceBase* AnimationSequenceBase, FName NotifyTrackName, FLinearColor TrackColor = FLinearColor::White);
+
+	/** Removes an Animation Notify Track from Animation Sequence by Name */
+	UFUNCTION(BlueprintCallable, Category = "AnimationBlueprintLibrary|AnimationNotifies")
+	static void RemoveAnimationNotifyTrack(UAnimSequenceBase* AnimationSequenceBase, FName NotifyTrackName);
+
+	/** Removes All Animation Notify Tracks from Animation Sequence */
+	UFUNCTION(BlueprintCallable, Category = "AnimationBlueprintLibrary|AnimationNotifies")
+	static void RemoveAllAnimationNotifyTracks(UAnimSequenceBase* AnimationSequenceBase);
+
+	/** Checks whether or not the given Track Name is a valid Animation Notify Track in the Animation Sequence */
+	UFUNCTION(BlueprintPure, Category = "AnimationBlueprintLibrary|Helpers")
+	static bool IsValidAnimNotifyTrackName(const UAnimSequenceBase* AnimationSequenceBase, FName NotifyTrackName);
+
+	static int32 GetTrackIndexForAnimationNotifyTrackName(const UAnimSequenceBase* AnimationSequenceBase, FName NotifyTrackName);
+	static const FAnimNotifyTrack& GetNotifyTrackByName(const UAnimSequenceBase* AnimationSequenceBase, FName NotifyTrackName);
+
+	/** Returns the actual trigger time for a NotifyEvent */
+	UFUNCTION(BlueprintPure, Category = "AnimationBlueprintLibrary|AnimationNotifies")
+	static float GetAnimNotifyEventTriggerTime(const FAnimNotifyEvent& NotifyEvent);
+	
+	/** Returns the duration for a NotifyEvent, only non-zero for Anim Notify States */
+	UFUNCTION(BlueprintPure, Category = "AnimationBlueprintLibrary|AnimationNotifies")
+	static float GetAnimNotifyEventDuration(const FAnimNotifyEvent& NotifyEvent);
+
+	/** Retrieves all Animation Sync Markers for the given Notify Track Name from the given Animation Sequence */
+	UFUNCTION(BlueprintPure, Category = "AnimationBlueprintLibrary|MarkerSyncing")
+	static void GetAnimationSyncMarkersForTrack(const UAnimSequence* AnimationSequence, FName NotifyTrackName, TArray<FAnimSyncMarker>& Markers);
+
+	/** Retrieves all Animation Notify Events for the given Notify Track Name from the given Animation Sequence */
+	UFUNCTION(BlueprintPure, Category = "AnimationBlueprintLibrary|NotifyEvents")
+	static void GetAnimationNotifyEventsForTrack(const UAnimSequenceBase* AnimationSequenceBase, FName NotifyTrackName, TArray<FAnimNotifyEvent>& Events);
+
+	// Curves
+
+	/** Adds an Animation Curve by Type and Name to the given Animation Sequence */
+	UFUNCTION(BlueprintCallable, Category = "AnimationBlueprintLibrary|Curves")
+	static void AddCurve(UAnimSequence* AnimationSequence, FName CurveName, ERawCurveTrackTypes CurveType = ERawCurveTrackTypes::RCT_Float, bool bMetaDataCurve = false);
+
+	/** Removes an Animation Curve by Name from the given Animation Sequence (Raw Animation Curves [Names] may not be removed from the Skeleton) */
+	UFUNCTION(BlueprintCallable, Category = "AnimationBlueprintLibrary|Curves")
+	static void RemoveCurve(UAnimSequence* AnimationSequence, FName CurveName, bool bRemoveNameFromSkeleton = false);
+
+	/** Removes all Animation Curve Data from the given Animation Sequence (Raw Animation Curves [Names] may not be removed from the Skeleton) */
+	UFUNCTION(BlueprintCallable, Category = "AnimationBlueprintLibrary|Curves")
+	static void RemoveAllCurveData(UAnimSequence* AnimationSequence);
+
+	/** Adds a Transformation Key to the specified Animation Curve inside of the given Animation Sequence */
+	UFUNCTION(BlueprintCallable, Category = "AnimationBlueprintLibrary|Curves")
+	static void AddTransformationCurveKey(UAnimSequence* AnimationSequence, FName CurveName, const float Time, const FTransform& Transform);
+
+	/** Adds a multiple of Transformation Keys to the specified Animation Curve inside of the given Animation Sequence */
+	UFUNCTION(BlueprintCallable, Category = "AnimationBlueprintLibrary|Curves")
+	static void AddTransformationCurveKeys(UAnimSequence* AnimationSequence, FName CurveName, const TArray<float>& Times, const TArray<FTransform>& Transforms);
+
+	/** Adds a Float Key to the specified Animation Curve inside of the given Animation Sequence */
+	UFUNCTION(BlueprintCallable, Category = "AnimationBlueprintLibrary|Curves")
+	static void AddFloatCurveKey(UAnimSequence* AnimationSequence, FName CurveName, const float Time, const float Value);
+
+	/** Adds a multiple of Float Keys to the specified Animation Curve inside of the given Animation Sequence */
+	UFUNCTION(BlueprintCallable, Category = "AnimationBlueprintLibrary|Curves")
+	static void AddFloatCurveKeys(UAnimSequence* AnimationSequence, FName CurveName, const TArray<float>& Times, const TArray<float>& Values);
+
+	/** Adds a Vector Key to the specified Animation Curve inside of the given Animation Sequence */
+	UFUNCTION(BlueprintCallable, Category = "AnimationBlueprintLibrary|Curves")
+	static void AddVectorCurveKey(UAnimSequence* AnimationSequence, FName CurveName, const float Time, const FVector Vector);
+
+	/** Adds a multiple of Vector Keys to the specified Animation Curve inside of the given Animation Sequence */
+	UFUNCTION(BlueprintCallable, Category = "AnimationBlueprintLibrary|Curves")
+	static void AddVectorCurveKeys(UAnimSequence* AnimationSequence, FName CurveName, const TArray<float>& Times, const TArray<FVector>& Vectors);
+
+	// Curve helper functions
+	template <typename DataType, typename CurveClass, ERawCurveTrackTypes CurveType>
+	static void AddCurveKeysInternal(UAnimSequence* AnimationSequence, FName CurveName, const TArray<float>& Times, const TArray<DataType>& KeyData);
+
+	// Returns true if successfully added, false if it was already existing
+	static bool AddCurveInternal(UAnimSequence* AnimationSequence, FName CurveName, FName ContainerName, int32 CurveFlags, ERawCurveTrackTypes SupportedCurveType);
+	static bool RemoveCurveInternal(UAnimSequence* AnimationSequence, FName CurveName, FName ContainerName, bool bRemoveNameFromSkeleton);
+
+	/** Checks whether or not the given Bone Name exist on the Skeleton referenced by the given Animation Sequence */
+	UFUNCTION(BlueprintCallable, Category = "AnimationBlueprintLibrary|Skeleton")
+	static void DoesBoneNameExist(UAnimSequence* AnimationSequence, FName BoneName, bool& bExists);
+
+	static bool DoesBoneNameExistInternal(USkeleton* Skeleton, FName BoneName);
+
+	/** Retrieves, a multiple of, Float Key(s) from the specified Animation Curve inside of the given Animation Sequence */
+	UFUNCTION(BlueprintCallable, Category = "AnimationBlueprintLibrary|Curves")
+	static void GetFloatKeys(UAnimSequence* AnimationSequence, FName CurveName, TArray<float>& Times, TArray<float>& Values);
+
+	/** Retrieves, a multiple of, Vector Key(s) from the specified Animation Curve inside of the given Animation Sequence */
+	UFUNCTION(BlueprintCallable, Category = "AnimationBlueprintLibrary|Curves")
+	static void GetVectorKeys(UAnimSequence* AnimationSequence, FName CurveName, TArray<float>& Times, TArray<FVector>& Values);
+
+	/** Retrieves, a multiple of, Transformation Key(s) from the specified Animation Curve inside of the given Animation Sequence */
+	UFUNCTION(BlueprintCallable, Category = "AnimationBlueprintLibrary|Curves")
+	static void GetTransformationKeys(UAnimSequence* AnimationSequence, FName CurveName, TArray<float>& Times, TArray<FTransform>& Values);
+
+	template <typename DataType, typename CurveClass, ERawCurveTrackTypes CurveType>
+	static void GetCurveKeysInternal(UAnimSequence* AnimationSequence, FName CurveName, TArray<float>& Times, TArray<DataType>& KeyData);
+
+	/** Ensures that any curve names that do not exist on the NewSkeleton are added to it, in which case the SmartName on the actual curve itself will also be updated */
+	UFUNCTION(BlueprintCallable, Category = "AnimationBlueprintLibrary|Curves")
+	static void CopyAnimationCurveNamesToSkeleton(USkeleton* OldSkeleton, USkeleton* NewSkeleton, UAnimSequenceBase* SequenceBase, ERawCurveTrackTypes CurveType);
+	
+	// Bone Tracks
+
+	/** Removes an Animation Curve by Name from the given Animation Sequence (Raw Animation Curves [Names] may not be removed from the Skeleton) 
+	 * 
+	 *	@param AnimationSequence : AnimSequence
+	 *	@param BoneName : Name of bone track user wants to remove
+	 *	@param bIncludeChildren : true if user wants to include all children of BoneName
+	 *  @param bFinalize : If you set this to true, it will trigger compression. If you set bFinalize to be false, you'll have to manually trigger Finalize. 
+	 */
+	UFUNCTION(BlueprintCallable, Category = "AnimationBlueprintLibrary|Bones")
+	static void RemoveBoneAnimation(UAnimSequence* AnimationSequence, FName BoneName, bool bIncludeChildren = true, bool bFinalize = true);
+
+	/** Removes all Animation Bone Track Data from the given Animation Sequence */
+	UFUNCTION(BlueprintCallable, Category = "AnimationBlueprintLibrary|Curves")
+	static void RemoveAllBoneAnimation(UAnimSequence* AnimationSequence);
+
+	/** Apply all the changes made to Bone Tracks to Finalize. This triggers recompression. Note that this is expensive, but will require to get correct compressed data */
+	UE_DEPRECATED(5.0, "FinalizeBoneAnimation has been deprecated, use UAnimDataController instead")
+	UFUNCTION(BlueprintCallable, Category = "AnimationBlueprintLibrary|Curves", meta=(DeprecatedFunction, DeprecationMessage="FinalizeBoneAnimation has been deprecated, use UAnimDataController instead"))
+	static void FinalizeBoneAnimation(UAnimSequence* AnimationSequence);
+
+	// Smart name helper functions
+
+	/** Checks whether or not the given Curve Name exist on the Skeleton referenced by the given Animation Sequence */
+	UFUNCTION(BlueprintCallable, Category = "AnimationBlueprintLibrary|Curves")
+	static bool DoesCurveExist(UAnimSequence* AnimationSequence, FName CurveName, ERawCurveTrackTypes CurveType);
+	static bool DoesSmartNameExist(UAnimSequence* AnimationSequence, FName Name);
+	
+	static FSmartName RetrieveSmartNameForCurve(const UAnimSequence* AnimationSequence, FName CurveName, FName ContainerName);
+	static bool RetrieveSmartNameForCurve(const UAnimSequence* AnimationSequence, FName CurveName, FName ContainerName, FSmartName& SmartName);
+	static FName RetrieveContainerNameForCurve(const UAnimSequence* AnimationSequence, FName CurveName);
+
+	// MetaData
+
+	/** Creates and Adds an instance of the specified MetaData Class to the given Animation Asset */
+	UFUNCTION(BlueprintCallable, Category = "AnimationBlueprintLibrary|MetaData")
+	static void AddMetaData(UAnimationAsset* AnimationAsset, TSubclassOf<UAnimMetaData> MetaDataClass, UAnimMetaData*& MetaDataInstance);
+
+	/** Adds an instance of the specified MetaData Class to the given Animation Asset (requires MetaDataObject's outer to be the Animation Asset) */
+	UFUNCTION(BlueprintCallable, Category = "AnimationBlueprintLibrary|MetaData")
+	static void AddMetaDataObject(UAnimationAsset* AnimationAsset, UAnimMetaData* MetaDataObject);
+
+	/** Removes all Meta Data from the given Animation Asset */
+	UFUNCTION(BlueprintCallable, Category = "AnimationBlueprintLibrary|MetaData")
+	static void RemoveAllMetaData(UAnimationAsset* AnimationAsset);
+
+	/** Removes the specified Meta Data Instance from the given Animation Asset */
+	UFUNCTION(BlueprintCallable, Category = "AnimationBlueprintLibrary|MetaData")
+	static void RemoveMetaData(UAnimationAsset* AnimationAsset, UAnimMetaData* MetaDataObject);
+
+	/** Removes all Meta Data Instance of the specified Class from the given Animation Asset */
+	UFUNCTION(BlueprintCallable, Category = "AnimationBlueprintLibrary|MetaData")
+	static void RemoveMetaDataOfClass(UAnimationAsset* AnimationAsset, TSubclassOf<UAnimMetaData> MetaDataClass);
+
+	/** Retrieves all Meta Data Instances from the given Animation Asset */
+	UFUNCTION(BlueprintPure, Category = "AnimationBlueprintLibrary|MetaData")
+	static void GetMetaData(const UAnimationAsset* AnimationAsset, TArray<UAnimMetaData*>& MetaData);
+
+	/** Retrieves all Meta Data Instances from the given Animation Asset */
+	UFUNCTION(BlueprintPure, Category = "AnimationBlueprintLibrary|MetaData")
+	static void GetMetaDataOfClass(const UAnimationAsset* AnimationAsset, TSubclassOf<UAnimMetaData> MetaDataClass, TArray<UAnimMetaData*>& MetaDataOfClass);
+
+	/** Checks whether or not the given Animation Asset contains Meta Data Instance of the specified Meta Data Class */
+	UFUNCTION(BlueprintPure, Category = "AnimationBlueprintLibrary|MetaData")
+	static bool ContainsMetaDataOfClass(const UAnimationAsset* AnimationAsset, TSubclassOf<UAnimMetaData> MetaDataClass);
+
+	// Poses
+
+	/** Retrieves Bone Pose data for the given Bone Name at the specified Time from the given Animation Sequence */
+	UFUNCTION(BlueprintPure, Category = "AnimationBlueprintLibrary|Pose")
+	static void GetBonePoseForTime(const UAnimSequenceBase* AnimationSequenceBase, FName BoneName, float Time, bool bExtractRootMotion, FTransform& Pose);
+
+	/** Retrieves Bone Pose data for the given Bone Name at the specified Frame from the given Animation Sequence */
+	UFUNCTION(BlueprintPure, Category = "AnimationBlueprintLibrary|Pose")
+	static void GetBonePoseForFrame(const UAnimSequenceBase* AnimationSequenceBase, FName BoneName, int32 Frame, bool bExtractRootMotion, FTransform& Pose);
+
+	/** Retrieves Bone Pose data for the given Bone Names at the specified Time from the given Animation Sequence */
+	UFUNCTION(BlueprintPure, Category = "AnimationBlueprintLibrary|Pose")
+	static void GetBonePosesForTime(const UAnimSequenceBase* AnimationSequenceBase, TArray<FName> BoneNames, float Time, bool bExtractRootMotion, TArray<FTransform>& Poses, const USkeletalMesh* PreviewMesh = nullptr);
+
+	/** Retrieves Bone Pose data for the given Bone Names at the specified Frame from the given Animation Sequence */
+	UFUNCTION(BlueprintPure, Category = "AnimationBlueprintLibrary|Pose")
+	static void GetBonePosesForFrame(const UAnimSequenceBase* AnimationSequenceBase, TArray<FName> BoneNames, int32 Frame, bool bExtractRootMotion, TArray<FTransform>& Poses, const USkeletalMesh* PreviewMesh = nullptr);
+
+	// Virtual bones
+
+	/** Adds a Virtual Bone between the Source and Target Bones to the given Animation Sequence */
+	UFUNCTION(BlueprintCallable, Category = "AnimationBlueprintLibrary|VirtualBones")
+	static void AddVirtualBone(const UAnimSequence* AnimationSequence, FName SourceBoneName, FName TargetBoneName, FName& VirtualBoneName);
+
+	/** Removes a Virtual Bone with the specified Bone Name from the given Animation Sequence */
+	UFUNCTION(BlueprintCallable, Category = "AnimationBlueprintLibrary|VirtualBones")
+	static void RemoveVirtualBone(const UAnimSequence* AnimationSequence, FName VirtualBoneName);
+
+	/** Removes Virtual Bones with the specified Bone Names from the given Animation Sequence */
+	UFUNCTION(BlueprintCallable, Category = "AnimationBlueprintLibrary|VirtualBones")
+	static void RemoveVirtualBones(const UAnimSequence* AnimationSequence, TArray<FName> VirtualBoneNames);
+
+	/** Removes all Virtual Bones from the given Animation Sequence */
+	UFUNCTION(BlueprintCallable, Category = "AnimationBlueprintLibrary|VirtualBones")
+	static void RemoveAllVirtualBones(const UAnimSequence* AnimationSequence);
+
+	static bool DoesVirtualBoneNameExistInternal(USkeleton* Skeleton, FName BoneName);
+
+	// Misc
+
+	/** Retrieves the Length of the given Animation Sequence */
+	UFUNCTION(BlueprintPure, Category = "AnimationBlueprintLibrary|Animation")
+	static void GetSequenceLength(const UAnimSequenceBase* AnimationSequenceBase, float& Length);
+
+	/** Retrieves the (Play) Rate Scale of the given Animation Sequence */
+	UFUNCTION(BlueprintPure, Category = "AnimationBlueprintLibrary|Animation")
+	static void GetRateScale(const UAnimSequenceBase* AnimationSequenceBase, float& RateScale);
+
+	/** Sets the (Play) Rate Scale for the given Animation Sequence */
+	UFUNCTION(BlueprintCallable, Category = "AnimationBlueprintLibrary|Animation")
+	static void SetRateScale(UAnimSequenceBase* AnimationSequenceBase, float RateScale);
+
+	/** Retrieves the Frame Index at the specified Time Value for the given Animation Sequence */
+	UFUNCTION(BlueprintPure, Category = "AnimationBlueprintLibrary|Helpers")
+	static void GetFrameAtTime(const UAnimSequenceBase* AnimationSequenceBase, const float Time, int32& Frame);
+
+	/** Retrieves the Time Value at the specified Frame Indexfor the given Animation Sequence */
+	UFUNCTION(BlueprintPure, Category = "AnimationBlueprintLibrary|Helpers")
+	static void GetTimeAtFrame(const UAnimSequenceBase* AnimationSequenceBase, const int32 Frame, float& Time);
+	
+	static float GetTimeAtFrameInternal(const UAnimSequenceBase* AnimationSequenceBase, const int32 Frame);
+
+	/** Checks whether or not the given Time Value lies within the given Animation Sequence's Length */
+	UFUNCTION(BlueprintPure, Category = "AnimationBlueprintLibrary|Helpers")
+	static void IsValidTime(const UAnimSequenceBase* AnimationSequenceBase, const float Time, bool& IsValid);
+
+	static bool IsValidTimeInternal(const UAnimSequenceBase* AnimationSequenceBase, const float Time);
+
+	/** Finds the Bone Path from the given Bone to the Root Bone */
+	UFUNCTION(BlueprintPure, Category = "AnimationBlueprintLibrary|Helpers")
+	static void FindBonePathToRoot(const UAnimSequenceBase* AnimationSequenceBase, FName BoneName, TArray<FName>& BonePath);
+};
diff --git a/Engine/Source/Editor/AnimationModifiers/AnimationModifiers.Build.cs b/Engine/Source/Editor/AnimationModifiers/AnimationModifiers.Build.cs
--- a/Engine/Source/Editor/AnimationModifiers/AnimationModifiers.Build.cs
+++ b/Engine/Source/Editor/AnimationModifiers/AnimationModifiers.Build.cs
@@ -25,7 +25,8 @@ public class AnimationModifiers : ModuleRules
                 "ClassViewer",
                 "AnimationEditor",
                 "SkeletonEditor",
-                "AssetRegistry"                
+                "AssetRegistry",
+                "AnimationBlueprintLibrary",
             }
 		);
     }
diff --git a/Engine/Source/Editor/AnimationModifiers/Private/AnimationBlueprintLibrary.cpp b/Engine/Source/Editor/AnimationModifiers/Private/AnimationBlueprintLibrary.cpp
deleted file mode 100644
diff --git a/Engine/Source/Editor/AnimationModifiers/Public/AnimationBlueprintLibrary.h b/Engine/Source/Editor/AnimationModifiers/Public/AnimationBlueprintLibrary.h
deleted file mode 100644
diff --git a/Engine/Source/Editor/UnrealEd/UnrealEd.Build.cs b/Engine/Source/Editor/UnrealEd/UnrealEd.Build.cs
--- a/Engine/Source/Editor/UnrealEd/UnrealEd.Build.cs
+++ b/Engine/Source/Editor/UnrealEd/UnrealEd.Build.cs
@@ -198,6 +198,7 @@ public class UnrealEd : ModuleRules
 				"TraceAnalysis",
 				"DeveloperSettings",
 				"AnimationModifiers",
+				"AnimationBlueprintLibrary",
 				"MaterialBaking",
 			}
 		);
