- Use Skeleton's reference skeleton rather than bone containers  when initializing FAnimPose
- Copy bone transforms manually using Skeleton index mapping rather than CopyFrom()

diff --git a/Engine/Source/Editor/AnimationBlueprintLibrary/Private/AnimPose.cpp b/Engine/Source/Editor/AnimationBlueprintLibrary/Private/AnimPose.cpp
--- a/Engine/Source/Editor/AnimationBlueprintLibrary/Private/AnimPose.cpp
+++ b/Engine/Source/Editor/AnimationBlueprintLibrary/Private/AnimPose.cpp
@@ -13,16 +13,23 @@ void FAnimPose::Init(const FBoneContainer& InBoneContainer)
 		
 	for (const FBoneIndexType BoneIndex : BoneContainer.GetBoneIndicesArray())
 	{			
-		const int32 SkeletonBoneIndex = BoneContainer.GetSkeletonIndex(FCompactPoseBoneIndex(BoneIndex));
-		const int32 ParentBoneIndex = BoneContainer.GetParentBoneIndex(SkeletonBoneIndex);
+		const FCompactPoseBoneIndex CompactIndex(BoneIndex);
+		const FCompactPoseBoneIndex CompactParentIndex = BoneContainer.GetParentBoneIndex(CompactIndex);
 
-		BoneIndices.Add(SkeletonBoneIndex);
-		ParentBoneIndices.Add(ParentBoneIndex);
+		const FReferenceSkeleton& RefSkeleton = BoneContainer.GetSkeletonAsset()->GetReferenceSkeleton();
 
-		const FReferenceSkeleton& RefSkeleton = BoneContainer.GetReferenceSkeleton();
-		BoneNames.Add(RefSkeleton.GetBoneName(SkeletonBoneIndex));
+		const int32 SkeletonBoneIndex = BoneContainer.GetSkeletonIndex(CompactIndex);
+		if (SkeletonBoneIndex != INDEX_NONE)
+		{
+			const int32 ParentBoneIndex = CompactParentIndex.GetInt() != INDEX_NONE ? BoneContainer.GetSkeletonIndex(CompactParentIndex) : INDEX_NONE;
+
+			BoneIndices.Add(SkeletonBoneIndex);
+			ParentBoneIndices.Add(ParentBoneIndex);
+
+			BoneNames.Add(RefSkeleton.GetBoneName(SkeletonBoneIndex));
 
-		RefLocalSpacePoses.Add(BoneContainer.GetRefPoseTransform(FCompactPoseBoneIndex(BoneIndex)));
+			RefLocalSpacePoses.Add(BoneContainer.GetRefPoseTransform(FCompactPoseBoneIndex(BoneIndex)));
+		}
 	}
 
 	TArray<bool> Processed;
@@ -367,9 +374,24 @@ void UAnimPoseExtensions::EvaluateAnimationBlueprintWithInputPose(const FAnimPos
 						{
 							if (FAnimNode_LinkedInputPose* InputNode = AnimInstance->GetLinkedInputPoseNode())
 							{
-								InputNode->CachedInputPose.SetBoneContainer(&AnimInstance->GetRequiredBones());
-								InputNode->CachedInputCurve.InitFrom(AnimInstance->GetRequiredBones());
-								InputNode->CachedInputPose.CopyBonesFrom(Pose.LocalSpacePoses);
+								const FBoneContainer& BoneContainer = AnimInstance->GetRequiredBones();
+								InputNode->CachedInputPose.SetBoneContainer(&BoneContainer);
+								InputNode->CachedInputCurve.InitFrom(BoneContainer);
+								InputNode->CachedInputPose.ResetToRefPose();
+
+								// Copy bone transform from input pose using skeleton index mapping
+								for (FCompactPoseBoneIndex CompactIndex : InputNode->CachedInputPose.ForEachBoneIndex())
+								{
+									const int32 SkeletonIndex = BoneContainer.GetSkeletonIndex(CompactIndex);
+									if (SkeletonIndex != INDEX_NONE)
+									{
+										const int32 Index = Pose.BoneIndices.IndexOfByKey(SkeletonIndex);
+										if (Index != INDEX_NONE)
+										{
+											InputNode->CachedInputPose[CompactIndex] = Pose.LocalSpacePoses[Index];
+										}
+									}
+								}
 					
 								OutPose.Init(AnimInstance->GetRequiredBones());
 
@@ -493,4 +515,4 @@ void UAnimPoseExtensions::GetAnimPoseAtTime(const UAnimSequenceBase* AnimationSe
 	{
 		UE_LOG(LogAnimationPoseScripting, Warning, TEXT("Invalid Animation Sequence supplied for GetBonePosesForTime"));
 	}
-}
\ No newline at end of file
+}
