Non-unity fixes

diff --git a/Engine/Source/Editor/AnimationBlueprintLibrary/Private/AnimPose.cpp b/Engine/Source/Editor/AnimationBlueprintLibrary/Private/AnimPose.cpp
--- a/Engine/Source/Editor/AnimationBlueprintLibrary/Private/AnimPose.cpp
+++ b/Engine/Source/Editor/AnimationBlueprintLibrary/Private/AnimPose.cpp
@@ -3,6 +3,8 @@
 #include "AnimPose.h"
 #include "PreviewScene.h"
 #include "Animation/AnimNode_LinkedInputPose.h"
+#include "Animation/AnimBlueprint.h"
+#include "Animation/AnimBlueprintGeneratedClass.h"
 
 DEFINE_LOG_CATEGORY_STATIC(LogAnimationPoseScripting, Verbose, All);
 
diff --git a/Engine/Source/Editor/AnimationBlueprintLibrary/Private/AnimationBlueprintLibrary.cpp b/Engine/Source/Editor/AnimationBlueprintLibrary/Private/AnimationBlueprintLibrary.cpp
--- a/Engine/Source/Editor/AnimationBlueprintLibrary/Private/AnimationBlueprintLibrary.cpp
+++ b/Engine/Source/Editor/AnimationBlueprintLibrary/Private/AnimationBlueprintLibrary.cpp
@@ -13,6 +13,8 @@
 #include "Animation/AnimCurveTypes.h"
 #include "BonePose.h" 
 #include "Algo/Transform.h"
+#include "Modules/ModuleManager.h"
+#include "Modules/ModuleInterface.h"
 
 #include "AnimationRuntime.h"
 #include "Animation/AnimData/AnimDataModel.h"
