Assert occurs when using Aim Animations from "/All/Game/Animation" folder within a Blend Space Analysis
#fix Actually init BlendedCurve during AnimPose generation/evaluation

diff --git a/Engine/Source/Editor/AnimationBlueprintLibrary/Private/AnimPose.cpp b/Engine/Source/Editor/AnimationBlueprintLibrary/Private/AnimPose.cpp
--- a/Engine/Source/Editor/AnimationBlueprintLibrary/Private/AnimPose.cpp
+++ b/Engine/Source/Editor/AnimationBlueprintLibrary/Private/AnimPose.cpp
@@ -500,6 +500,7 @@ void UAnimPoseExtensions::GetAnimPoseAtTime(const UAnimSequenceBase* AnimationSe
 			Pose.Init(CompactPose.GetBoneContainer());
 
 			FBlendedCurve Curve;
+			Curve.InitFrom(RequiredBones);
 			UE::Anim::FStackAttributeContainer Attributes;
 
 			FAnimationPoseData PoseData(CompactPose, Curve, Attributes);
