AnimationBlueprintLibrary::DoesBoneNameExistInternal logic incorrect

diff --git a/Engine/Source/Editor/AnimationBlueprintLibrary/Private/AnimationBlueprintLibrary.cpp b/Engine/Source/Editor/AnimationBlueprintLibrary/Private/AnimationBlueprintLibrary.cpp
--- a/Engine/Source/Editor/AnimationBlueprintLibrary/Private/AnimationBlueprintLibrary.cpp
+++ b/Engine/Source/Editor/AnimationBlueprintLibrary/Private/AnimationBlueprintLibrary.cpp
@@ -1495,7 +1495,7 @@ bool UAnimationBlueprintLibrary::RemoveCurveInternal(UAnimSequence* AnimationSeq
 				if (ContainerName == USkeleton::AnimTrackCurveMappingName)
 				{
 					// Make sure we do not remove bone names
-					bValidToRemove = DoesBoneNameExistInternal(Skeleton, CurveName);
+					bValidToRemove = DoesBoneCurveNameExistInternal(Skeleton, CurveName);
 				}
 
 				if (bValidToRemove)
@@ -1544,6 +1544,12 @@ void UAnimationBlueprintLibrary::DoesBoneNameExist(UAnimSequence* AnimationSeque
 }
 
 bool UAnimationBlueprintLibrary::DoesBoneNameExistInternal(USkeleton* Skeleton, FName BoneName)
+{
+	checkf(Skeleton != nullptr, TEXT("Invalid Skeleton ptr"));
+	return Skeleton->GetReferenceSkeleton().FindBoneIndex(BoneName) != INDEX_NONE;
+}
+
+bool UAnimationBlueprintLibrary::DoesBoneCurveNameExistInternal(USkeleton* Skeleton, FName BoneName)
 {
 	checkf(Skeleton != nullptr, TEXT("Invalid Skeleton ptr"));
 	return Skeleton->GetUIDByName(USkeleton::AnimTrackCurveMappingName, BoneName) != SmartName::MaxUID;
diff --git a/Engine/Source/Editor/AnimationBlueprintLibrary/Public/AnimationBlueprintLibrary.h b/Engine/Source/Editor/AnimationBlueprintLibrary/Public/AnimationBlueprintLibrary.h
--- a/Engine/Source/Editor/AnimationBlueprintLibrary/Public/AnimationBlueprintLibrary.h
+++ b/Engine/Source/Editor/AnimationBlueprintLibrary/Public/AnimationBlueprintLibrary.h
@@ -313,6 +313,7 @@ public:
 	static void DoesBoneNameExist(UAnimSequence* AnimationSequence, FName BoneName, bool& bExists);
 
 	static bool DoesBoneNameExistInternal(USkeleton* Skeleton, FName BoneName);
+	static bool DoesBoneCurveNameExistInternal(USkeleton* Skeleton, FName BoneName);
 
 	/** Retrieves, a multiple of, Float Key(s) from the specified Animation Curve inside of the given Animation Sequence */
 	UFUNCTION(BlueprintCallable, Category = "AnimationBlueprintLibrary|Curves")
