Fix for BlendSpace analysis fails with additive animations (code from Jurre)

diff --git a/Engine/Source/Editor/AnimationBlueprintLibrary/Private/AnimPose.cpp b/Engine/Source/Editor/AnimationBlueprintLibrary/Private/AnimPose.cpp
--- a/Engine/Source/Editor/AnimationBlueprintLibrary/Private/AnimPose.cpp
+++ b/Engine/Source/Editor/AnimationBlueprintLibrary/Private/AnimPose.cpp
@@ -5,6 +5,7 @@
 #include "Animation/AnimNode_LinkedInputPose.h"
 #include "Animation/AnimBlueprint.h"
 #include "Animation/AnimBlueprintGeneratedClass.h"
+#include "Animation/AnimSequence.h"
 
 DEFINE_LOG_CATEGORY_STATIC(LogAnimationPoseScripting, Verbose, All);
 
@@ -507,7 +508,29 @@ void UAnimPoseExtensions::GetAnimPoseAtTime(const UAnimSequenceBase* AnimationSe
 			FAnimExtractContext Context(Time, EvaluationOptions.bExtractRootMotion);
 
 			AnimationSequenceBase->GetAnimationPose(PoseData, Context);
-			Pose.SetPose(CompactPose);
+
+			if (AnimationSequenceBase->IsValidAdditive())
+			{
+				const UAnimSequence* AnimSequence = Cast<const UAnimSequence>(AnimationSequenceBase);
+				FCompactPose BasePose;
+				BasePose.SetBoneContainer(&RequiredBones);
+
+				FBlendedCurve BaseCurve;
+				BaseCurve.InitFrom(RequiredBones);
+				UE::Anim::FStackAttributeContainer BaseAttributes;
+				
+				FAnimationPoseData BasePoseData(BasePose, BaseCurve, BaseAttributes);
+				AnimSequence->GetAdditiveBasePose(BasePoseData, Context);
+
+				FAnimationRuntime::AccumulateAdditivePose(BasePoseData, PoseData, 1.f, AnimSequence->GetAdditiveAnimType());
+				BasePose.NormalizeRotations();
+				
+				Pose.SetPose(BasePose);
+			}
+			else
+			{
+				Pose.SetPose(CompactPose);
+			}
 		}
 		else
 		{
diff --git a/Engine/Source/Editor/Persona/Private/PersonaBlendSpaceAnalysis.cpp b/Engine/Source/Editor/Persona/Private/PersonaBlendSpaceAnalysis.cpp
--- a/Engine/Source/Editor/Persona/Private/PersonaBlendSpaceAnalysis.cpp
+++ b/Engine/Source/Editor/Persona/Private/PersonaBlendSpaceAnalysis.cpp
@@ -5,7 +5,6 @@
 #include "Engine/SkeletalMeshSocket.h"
 #include "Animation/BlendSpace.h"
 #include "Animation/AnimSequence.h"
-#include "Animation/AnimData/AnimDataModel.h"
 #include "Animation/BoneSocketReference.h"
 
 #define LOCTEXT_NAMESPACE "BlendSpaceAnalysis"
@@ -48,8 +47,6 @@ bool BlendSpaceAnalysis::GetBoneInfo(const UAnimSequence&     Animation,
 									 FTransform&              BoneOffset, 
 									 FName&                   BoneName)
 {
-	const UAnimDataModel* AnimDataModel = Animation.GetDataModel();
-
 	if (BoneSocket.bUseSocket)
 	{
 		const USkeleton* Skeleton = Animation.GetSkeleton();
