AnimationBlueprintLibrary: add function for evaluating timecode attributes for the root bone in an anim sequence

A specially designated set of attribute curves can be authored on a bone
to encode timecode metadata for an animation sequence. Each attribute
encodes a different component of a complete timecode and subframe value
("TCHour", "TCMinute", "TCSecond", "TCFrame", and "TCSubframe"). The
convention is to author these attribute curves on the root bone of the animation.

This adds a function for evaluating those attributes at a particular time, and the result
is returned as a qualified frame time. If the anim sequence has an import file frame rate
specified, then that rate is used for the returned qualified frame time to more closely
match the original source of the animation. Otherwise, the sampling frame rate of the
anim sequence is used. If the root bone does not have any timecode attributes or they
otherwise cannot be evaluated (for example, because they have no keys), the function
returns false, and the qualified frame time passed by reference is unmodified.

Finally, this new function is used for the frame time hint of skeletal animation tracks in
Sequencer. Anim sequences that have these attributes authored will show the authored
timecode/frame values in the frame time hint when they are selected. This enables more
easily correlating a point in time in the anim sequence with where it came from in an
external DCC.

diff --git a/Engine/Source/Editor/AnimationBlueprintLibrary/Private/AnimationBlueprintLibrary.cpp b/Engine/Source/Editor/AnimationBlueprintLibrary/Private/AnimationBlueprintLibrary.cpp
--- a/Engine/Source/Editor/AnimationBlueprintLibrary/Private/AnimationBlueprintLibrary.cpp
+++ b/Engine/Source/Editor/AnimationBlueprintLibrary/Private/AnimationBlueprintLibrary.cpp
@@ -8,11 +8,17 @@
 #include "Animation/AnimCurveCompressionSettings.h"
 #include "Animation/AnimMetaData.h"
 #include "Animation/AnimNotifies/AnimNotifyState.h"
+#include "Animation/BuiltInAttributeTypes.h"
 #include "Animation/Skeleton.h"
 #include "Animation/AnimNotifies/AnimNotify.h"
 #include "Animation/AnimCurveTypes.h"
+#include "Animation/AnimationSettings.h"
 #include "BonePose.h" 
 #include "Algo/Transform.h"
+#include "Misc/FrameRate.h"
+#include "Misc/FrameTime.h"
+#include "Misc/QualifiedFrameTime.h"
+#include "Misc/Timecode.h"
 #include "Modules/ModuleManager.h"
 #include "Modules/ModuleInterface.h"
 
@@ -2164,6 +2170,145 @@ bool UAnimationBlueprintLibrary::IsValidTimeInternal(const UAnimSequenceBase* An
 	return FMath::IsWithinInclusive(Time, 0.0f, AnimationSequenceBase->GetPlayLength());
 }
 
+bool UAnimationBlueprintLibrary::EvaluateRootBoneTimecodeAttributesAtTime(const UAnimSequenceBase* AnimationSequenceBase, const float EvalTime, FQualifiedFrameTime& OutQualifiedFrameTime)
+{
+	if (!AnimationSequenceBase)
+	{
+		return false;
+	}
+
+	const UAnimDataModel* AnimDataModel = AnimationSequenceBase->GetDataModel();
+	if (!AnimDataModel)
+	{
+		return false;
+	}
+
+	const int32 RootBoneTrackIndex = 0;
+	const FBoneAnimationTrack* RootBoneTrack = AnimDataModel->FindBoneTrackByIndex(RootBoneTrackIndex);
+	if (!RootBoneTrack)
+	{
+		return false;
+	}
+
+	const FName& RootBoneName = RootBoneTrack->Name;
+
+	TArray<const FAnimatedBoneAttribute*> RootBoneAttributes;
+	AnimDataModel->GetAttributesForBone(RootBoneName, RootBoneAttributes);
+
+	FName TCHourAttrName(TEXT("TCHour"));
+	FName TCMinuteAttrName(TEXT("TCMinute"));
+	FName TCSecondAttrName(TEXT("TCSecond"));
+	FName TCFrameAttrName(TEXT("TCFrame"));
+	FName TCSubframeAttrName(TEXT("TCSubframe"));
+
+	if (const UAnimationSettings* AnimationSettings = UAnimationSettings::Get())
+	{
+		TCHourAttrName = AnimationSettings->BoneTimecodeCustomAttributeNameSettings.HourAttributeName;
+		TCMinuteAttrName = AnimationSettings->BoneTimecodeCustomAttributeNameSettings.MinuteAttributeName;
+		TCSecondAttrName = AnimationSettings->BoneTimecodeCustomAttributeNameSettings.SecondAttributeName;
+		TCFrameAttrName = AnimationSettings->BoneTimecodeCustomAttributeNameSettings.FrameAttributeName;
+		TCSubframeAttrName = AnimationSettings->BoneTimecodeCustomAttributeNameSettings.SubframeAttributeName;
+	}
+
+	const TArray<FName> TimecodeBoneAttributeNames = { TCHourAttrName, TCMinuteAttrName, TCSecondAttrName, TCFrameAttrName, TCSubframeAttrName };
+
+	bool bHasTimecodeBoneAttributes = false;
+
+	FTimecode Timecode;
+	float SubFrame = 0.0f;
+
+	for (const FAnimatedBoneAttribute* RootBoneAttribute : RootBoneAttributes)
+	{
+		if (!RootBoneAttribute)
+		{
+			continue;
+		}
+
+		const FName& BoneAttributeName = RootBoneAttribute->Identifier.GetName();
+
+		// Avoid evaluating non-timecode bone attributes.
+		if (!TimecodeBoneAttributeNames.Contains(BoneAttributeName))
+		{
+			continue;
+		}
+
+		if (!RootBoneAttribute->Curve.CanEvaluate())
+		{
+			continue;
+		}
+
+		float FloatValue = 0.0f;
+		int32 IntValue = 0;
+
+		// Support timecode attribute curves that are either float-typed or integer-typed.
+		if (RootBoneAttribute->Curve.GetScriptStruct() == FFloatAnimationAttribute::StaticStruct())
+		{
+			const FFloatAnimationAttribute EvaluatedAttribute = RootBoneAttribute->Curve.Evaluate<FFloatAnimationAttribute>(EvalTime);
+			FloatValue = EvaluatedAttribute.Value;
+			IntValue = static_cast<int32>(FloatValue);
+		}
+		else if (RootBoneAttribute->Curve.GetScriptStruct() == FIntegerAnimationAttribute::StaticStruct())
+		{
+			const FIntegerAnimationAttribute EvaluatedAttribute = RootBoneAttribute->Curve.Evaluate<FIntegerAnimationAttribute>(EvalTime);
+			IntValue = EvaluatedAttribute.Value;
+			FloatValue = static_cast<float>(IntValue);
+		}
+		else
+		{
+			continue;
+		}
+
+		if (BoneAttributeName.IsEqual(TCHourAttrName))
+		{
+			Timecode.Hours = IntValue;
+			bHasTimecodeBoneAttributes = true;
+		}
+		else if (BoneAttributeName.IsEqual(TCMinuteAttrName))
+		{
+			Timecode.Minutes = IntValue;
+			bHasTimecodeBoneAttributes = true;
+		}
+		else if (BoneAttributeName.IsEqual(TCSecondAttrName))
+		{
+			Timecode.Seconds = IntValue;
+			bHasTimecodeBoneAttributes = true;
+		}
+		else if (BoneAttributeName.IsEqual(TCFrameAttrName))
+		{
+			Timecode.Frames = IntValue;
+			bHasTimecodeBoneAttributes = true;
+		}
+		else if (BoneAttributeName.IsEqual(TCSubframeAttrName))
+		{
+			SubFrame = FloatValue;
+			bHasTimecodeBoneAttributes = true;
+		}
+	}
+
+	if (!bHasTimecodeBoneAttributes)
+	{
+		return false;
+	}
+
+	// We'll fall back on the sampling frame rate if we can't determine the
+	// original source frame rate.
+	FFrameRate FrameRate = AnimationSequenceBase->GetSamplingFrameRate();
+
+	if (const UAnimSequence* AnimSequence = AnimDataModel->GetAnimationSequence())
+	{
+		if (AnimSequence->ImportFileFramerate > 0.0f)
+		{
+			FrameRate = FFrameRate(static_cast<int32>(AnimSequence->ImportFileFramerate), 1);
+		}
+	}
+
+	OutQualifiedFrameTime = FQualifiedFrameTime(
+		FFrameTime(Timecode.ToFrameNumber(FrameRate), SubFrame),
+		FrameRate);
+
+	return true;
+}
+
 void UAnimationBlueprintLibrary::FindBonePathToRoot(const UAnimSequenceBase* AnimationSequenceBase, FName BoneName, TArray<FName>& BonePath)
 {
 	BonePath.Empty();
diff --git a/Engine/Source/Editor/AnimationBlueprintLibrary/Public/AnimationBlueprintLibrary.h b/Engine/Source/Editor/AnimationBlueprintLibrary/Public/AnimationBlueprintLibrary.h
--- a/Engine/Source/Editor/AnimationBlueprintLibrary/Public/AnimationBlueprintLibrary.h
+++ b/Engine/Source/Editor/AnimationBlueprintLibrary/Public/AnimationBlueprintLibrary.h
@@ -10,6 +10,7 @@
 
 #include "AnimationBlueprintLibrary.generated.h"
 
+struct FQualifiedFrameTime;
 struct FRawAnimSequenceTrack;
 class UAnimCompress;
 class USkeleton;
@@ -469,6 +470,18 @@ public:
 
 	static bool IsValidTimeInternal(const UAnimSequenceBase* AnimationSequenceBase, const float Time);
 
+	/** Evaluates timecode attributes (e.g. "TCFrame", "TCSecond", etc.) of the root bone and returns the resulting qualified frame time.
+	 *
+	 *  @param AnimationSequenceBase: Anim sequence for which to evaluate the root bone attributes.
+	 *  @param EvalTime: Time (in seconds) at which to evaluate the timecode bone attributes.
+	 *  @param OutQualifiedFrameTime: Resulting qualified frame time from evaluation. If the anim sequence has an import file frame rate
+	 *      set, then that will be used as the frame rate of the qualified frame time. Otherwise, the sampling frame rate of the anim
+	 *      sequence is used. If no timecode attributes are present on the bone or if none can be evaluated, the passed object will not be modified.
+	 *  @return: true if the root bone had timecode attributes that could be evaluated and a qualified frame time was set, or false otherwise.
+	 */
+	UFUNCTION(BlueprintPure, Category = "AnimationBlueprintLibrary|Helpers")
+	static bool EvaluateRootBoneTimecodeAttributesAtTime(const UAnimSequenceBase* AnimationSequenceBase, const float EvalTime, FQualifiedFrameTime& OutQualifiedFrameTime);
+
 	/** Finds the Bone Path from the given Bone to the Root Bone */
 	UFUNCTION(BlueprintPure, Category = "AnimationBlueprintLibrary|Helpers")
 	static void FindBonePathToRoot(const UAnimSequenceBase* AnimationSequenceBase, FName BoneName, TArray<FName>& BonePath);
diff --git a/Engine/Source/Editor/MovieSceneTools/MovieSceneTools.Build.cs b/Engine/Source/Editor/MovieSceneTools/MovieSceneTools.Build.cs
--- a/Engine/Source/Editor/MovieSceneTools/MovieSceneTools.Build.cs
+++ b/Engine/Source/Editor/MovieSceneTools/MovieSceneTools.Build.cs
@@ -36,6 +36,7 @@ public class MovieSceneTools : ModuleRules
 		PrivateDependencyModuleNames.AddRange(
 			new string[] {
                 "ActorPickerMode",
+				"AnimationBlueprintLibrary",
 				"AppFramework",
 				"CinematicCamera",
 				"DataLayerEditor",
diff --git a/Engine/Source/Editor/MovieSceneTools/Private/TrackEditors/SkeletalAnimationTrackEditor.cpp b/Engine/Source/Editor/MovieSceneTools/Private/TrackEditors/SkeletalAnimationTrackEditor.cpp
--- a/Engine/Source/Editor/MovieSceneTools/Private/TrackEditors/SkeletalAnimationTrackEditor.cpp
+++ b/Engine/Source/Editor/MovieSceneTools/Private/TrackEditors/SkeletalAnimationTrackEditor.cpp
@@ -37,6 +37,7 @@
 #include "SequencerTimeSliderController.h"
 #include "FrameNumberDisplayFormat.h"
 #include "FrameNumberNumericInterface.h"
+#include "AnimationBlueprintLibrary.h"
 #include "AnimationEditorUtils.h"
 #include "Factories/PoseAssetFactory.h"
 #include "Misc/MessageDialog.h"
@@ -619,7 +620,13 @@ int32 FSkeletalAnimationSection::OnPaintSection( FSequencerSectionPainter& Paint
 			// Draw the current time next to the scrub handle
 			const double AnimTime = Section.MapTimeToAnimation(CurrentTime, TickResolution);
 			const FFrameRate SamplingFrameRate = Section.Params.Animation->GetSamplingFrameRate();
-			const FFrameTime FrameTime = SamplingFrameRate.AsFrameTime(AnimTime);
+
+			FQualifiedFrameTime HintFrameTime;
+			if (!UAnimationBlueprintLibrary::EvaluateRootBoneTimecodeAttributesAtTime(Section.Params.Animation, static_cast<float>(AnimTime), HintFrameTime))
+			{
+				const FFrameTime FrameTime = SamplingFrameRate.AsFrameTime(AnimTime);
+				HintFrameTime = FQualifiedFrameTime(FrameTime, SamplingFrameRate);
+			}
 
 			// Get the desired frame display format and zero padding from
 			// the sequencer settings, if possible.
@@ -633,12 +640,12 @@ int32 FSkeletalAnimationSection::OnPaintSection( FSequencerSectionPainter& Paint
 
 			// No frame rate conversion necessary since we're displaying
 			// the source frame time/rate.
-			const TAttribute<FFrameRate> TickResolutionAttr(SamplingFrameRate);
-			const TAttribute<FFrameRate> DisplayRateAttr(SamplingFrameRate);
+			const TAttribute<FFrameRate> TickResolutionAttr(HintFrameTime.Rate);
+			const TAttribute<FFrameRate> DisplayRateAttr(HintFrameTime.Rate);
 
 			const FFrameNumberInterface FrameNumberInterface(DisplayFormatAttr, ZeroPadFrameNumbersAttr, TickResolutionAttr, DisplayRateAttr);
 
-			DrawFrameTimeHint(Painter, CurrentTime, FrameTime, &FrameNumberInterface);
+			DrawFrameTimeHint(Painter, CurrentTime, HintFrameTime.Time, &FrameNumberInterface);
 		}
 	}
 	
