changed transform  functions to return const references

diff --git a/Engine/Source/Editor/AnimationBlueprintLibrary/Private/AnimPose.cpp b/Engine/Source/Editor/AnimationBlueprintLibrary/Private/AnimPose.cpp
--- a/Engine/Source/Editor/AnimationBlueprintLibrary/Private/AnimPose.cpp
+++ b/Engine/Source/Editor/AnimationBlueprintLibrary/Private/AnimPose.cpp
@@ -13,14 +13,14 @@ void FAnimPose::Init(const FBoneContainer& InBoneContainer)
 {
 	Reset();
 	BoneContainer = InBoneContainer;
-		
+
+	const FReferenceSkeleton& RefSkeleton = BoneContainer.GetSkeletonAsset()->GetReferenceSkeleton();
+	
 	for (const FBoneIndexType BoneIndex : BoneContainer.GetBoneIndicesArray())
 	{			
 		const FCompactPoseBoneIndex CompactIndex(BoneIndex);
 		const FCompactPoseBoneIndex CompactParentIndex = BoneContainer.GetParentBoneIndex(CompactIndex);
 
-		const FReferenceSkeleton& RefSkeleton = BoneContainer.GetSkeletonAsset()->GetReferenceSkeleton();
-
 		const int32 SkeletonBoneIndex = BoneContainer.GetSkeletonIndex(CompactIndex);
 		if (SkeletonBoneIndex != INDEX_NONE)
 		{
@@ -195,7 +195,7 @@ void UAnimPoseExtensions::GetBoneNames(const FAnimPose& Pose, TArray<FName>& Bon
 	Bones.Append(Pose.BoneNames);
 }
 
-FTransform UAnimPoseExtensions::GetBonePose(const FAnimPose& Pose, FName BoneName, EAnimPoseSpaces Space /*= EAnimPoseSpaces::Local*/)
+const FTransform& UAnimPoseExtensions::GetBonePose(const FAnimPose& Pose, FName BoneName, EAnimPoseSpaces Space /*= EAnimPoseSpaces::Local*/)
 {
 	if (Pose.IsValid())
 	{
@@ -247,7 +247,7 @@ void UAnimPoseExtensions::SetBonePose(FAnimPose& Pose, FTransform Transform, FNa
 	}
 }
 	
-FTransform UAnimPoseExtensions::GetRefBonePose(const FAnimPose& Pose, FName BoneName, EAnimPoseSpaces Space /*= EAnimPoseSpaces::Local*/)
+const FTransform& UAnimPoseExtensions::GetRefBonePose(const FAnimPose& Pose, FName BoneName, EAnimPoseSpaces Space /*= EAnimPoseSpaces::Local*/)
 {
 	if (Pose.IsValid())
 	{
diff --git a/Engine/Source/Editor/AnimationBlueprintLibrary/Public/AnimPose.h b/Engine/Source/Editor/AnimationBlueprintLibrary/Public/AnimPose.h
--- a/Engine/Source/Editor/AnimationBlueprintLibrary/Public/AnimPose.h
+++ b/Engine/Source/Editor/AnimationBlueprintLibrary/Public/AnimPose.h
@@ -147,8 +147,8 @@ public:
 	* @return	Transform in requested space for bone if found, otherwise return identity transform
 	*/
 	UFUNCTION(BlueprintPure, meta = (ScriptMethod), Category = "Animation|Pose")
-    static FTransform GetBonePose(UPARAM(ref) const FAnimPose& Pose, FName BoneName, EAnimPoseSpaces Space = EAnimPoseSpaces::Local);
-
+    static const FTransform& GetBonePose(UPARAM(ref) const FAnimPose& Pose, FName BoneName, EAnimPoseSpaces Space = EAnimPoseSpaces::Local);
+	
 	/**
 	* Sets the transform for the provided bone name for a pose
 	*
@@ -170,7 +170,7 @@ public:
 	* @return	Transform in requested space for bone if found, otherwise return identity transform
 	*/
 	UFUNCTION(BlueprintPure, meta = (ScriptMethod), Category = "Animation|Pose")
-    static FTransform GetRefBonePose(UPARAM(ref) const FAnimPose& Pose, FName BoneName, EAnimPoseSpaces Space = EAnimPoseSpaces::Local);
+    static const FTransform& GetRefBonePose(UPARAM(ref) const FAnimPose& Pose, FName BoneName, EAnimPoseSpaces Space = EAnimPoseSpaces::Local);
 
 	/**
 	* Retrieves the relative transform between the two provided bone names
