Animation: add "TCRate" timecode custom attribute for specifying the timecode rate

In some workflows when authoring timecode metadata, the timecode rate may be
different from the animation or capture frame rate, for example when capturing
"high" frame rate data at 120 frames per second but recording SMPTE timecode
at 30 frames per second.

This adds a "TCRate" custom attribute setting that allows authoring metadata that
specifies a timecode rate that is different from the animation data frame rate.

With this initial introduction of the new setting, the custom attribute is expected to
be integer or float typed and should represent the timecode rate in frames per second.
Drop frame is not yet supported.

diff --git a/Engine/Config/BaseEngine.ini b/Engine/Config/BaseEngine.ini
--- a/Engine/Config/BaseEngine.ini
+++ b/Engine/Config/BaseEngine.ini
@@ -2826,7 +2826,7 @@ HMDWornMovementThreshold = 50.0
 
 [/Script/Engine.AnimationSettings]
 bStripAnimationDataOnDedicatedServer=False
-BoneTimecodeCustomAttributeNameSettings=(HourAttributeName="TCHour",MinuteAttributeName="TCMinute",SecondAttributeName="TCSecond",FrameAttributeName="TCFrame",SubframeAttributeName="TCSubframe",TakenameAttributeName="Takename")
+BoneTimecodeCustomAttributeNameSettings=(HourAttributeName="TCHour",MinuteAttributeName="TCMinute",SecondAttributeName="TCSecond",FrameAttributeName="TCFrame",SubframeAttributeName="TCSubframe",RateAttributeName="TCRate",TakenameAttributeName="Takename")
 
 [Animation.DefaultObjectSettings]
 BoneCompressionSettings="/Engine/Animation/DefaultAnimBoneCompressionSettings"
diff --git a/Engine/Source/Editor/AnimationBlueprintLibrary/AnimationBlueprintLibrary.build.cs b/Engine/Source/Editor/AnimationBlueprintLibrary/AnimationBlueprintLibrary.build.cs
--- a/Engine/Source/Editor/AnimationBlueprintLibrary/AnimationBlueprintLibrary.build.cs
+++ b/Engine/Source/Editor/AnimationBlueprintLibrary/AnimationBlueprintLibrary.build.cs
@@ -14,6 +14,7 @@ public class AnimationBlueprintLibrary : ModuleRules
 				"Engine",
 				"Kismet",
 				"AnimGraph",
+				"TimeManagement"
 			}
 		);
 
diff --git a/Engine/Source/Editor/AnimationBlueprintLibrary/Private/AnimationBlueprintLibrary.cpp b/Engine/Source/Editor/AnimationBlueprintLibrary/Private/AnimationBlueprintLibrary.cpp
--- a/Engine/Source/Editor/AnimationBlueprintLibrary/Private/AnimationBlueprintLibrary.cpp
+++ b/Engine/Source/Editor/AnimationBlueprintLibrary/Private/AnimationBlueprintLibrary.cpp
@@ -14,6 +14,7 @@
 #include "Animation/AnimCurveTypes.h"
 #include "Animation/AnimationSettings.h"
 #include "BonePose.h" 
+#include "CommonFrameRates.h"
 #include "Algo/Transform.h"
 #include "Misc/FrameRate.h"
 #include "Misc/FrameTime.h"
@@ -2200,6 +2201,7 @@ bool UAnimationBlueprintLibrary::EvaluateRootBoneTimecodeAttributesAtTime(const
 	FName TCSecondAttrName(TEXT("TCSecond"));
 	FName TCFrameAttrName(TEXT("TCFrame"));
 	FName TCSubframeAttrName(TEXT("TCSubframe"));
+	FName TCRateAttrName(TEXT("TCRate"));
 
 	if (const UAnimationSettings* AnimationSettings = UAnimationSettings::Get())
 	{
@@ -2208,15 +2210,18 @@ bool UAnimationBlueprintLibrary::EvaluateRootBoneTimecodeAttributesAtTime(const
 		TCSecondAttrName = AnimationSettings->BoneTimecodeCustomAttributeNameSettings.SecondAttributeName;
 		TCFrameAttrName = AnimationSettings->BoneTimecodeCustomAttributeNameSettings.FrameAttributeName;
 		TCSubframeAttrName = AnimationSettings->BoneTimecodeCustomAttributeNameSettings.SubframeAttributeName;
+		TCRateAttrName = AnimationSettings->BoneTimecodeCustomAttributeNameSettings.RateAttributeName;
 	}
 
-	const TArray<FName> TimecodeBoneAttributeNames = { TCHourAttrName, TCMinuteAttrName, TCSecondAttrName, TCFrameAttrName, TCSubframeAttrName };
+	const TArray<FName> TimecodeBoneAttributeNames = { TCHourAttrName, TCMinuteAttrName, TCSecondAttrName, TCFrameAttrName, TCSubframeAttrName, TCRateAttrName };
 
 	bool bHasTimecodeBoneAttributes = false;
 
 	FTimecode Timecode;
 	float SubFrame = 0.0f;
 
+	double TimecodeRateAsDecimal = 0.0;
+
 	for (const FAnimatedBoneAttribute* RootBoneAttribute : RootBoneAttributes)
 	{
 		if (!RootBoneAttribute)
@@ -2283,6 +2288,12 @@ bool UAnimationBlueprintLibrary::EvaluateRootBoneTimecodeAttributesAtTime(const
 			SubFrame = FloatValue;
 			bHasTimecodeBoneAttributes = true;
 		}
+		else if (BoneAttributeName.IsEqual(TCRateAttrName))
+		{
+			TimecodeRateAsDecimal = FloatValue;
+			// Don't consider this attribute when determining whether timecode value attributes are
+			// present, since it can't be useful on its own.
+		}
 	}
 
 	if (!bHasTimecodeBoneAttributes)
@@ -2294,7 +2305,18 @@ bool UAnimationBlueprintLibrary::EvaluateRootBoneTimecodeAttributesAtTime(const
 	// original source frame rate.
 	FFrameRate FrameRate = AnimationSequenceBase->GetSamplingFrameRate();
 
-	if (const UAnimSequence* AnimSequence = AnimDataModel->GetAnimationSequence())
+	if (TimecodeRateAsDecimal > 1.0)
+	{
+		if (const FCommonFrameRateInfo* TimecodeRateInfo = FCommonFrameRates::Find(TimecodeRateAsDecimal))
+		{
+			FrameRate = TimecodeRateInfo->FrameRate;
+		}
+		else
+		{
+			FrameRate = FFrameRate(FMath::TruncToInt(TimecodeRateAsDecimal), 1);
+		}
+	}
+	else if (const UAnimSequence* AnimSequence = AnimDataModel->GetAnimationSequence())
 	{
 		if (AnimSequence->ImportFileFramerate > 0.0f)
 		{
diff --git a/Engine/Source/Runtime/Engine/Classes/Animation/CustomAttributes.h b/Engine/Source/Runtime/Engine/Classes/Animation/CustomAttributes.h
--- a/Engine/Source/Runtime/Engine/Classes/Animation/CustomAttributes.h
+++ b/Engine/Source/Runtime/Engine/Classes/Animation/CustomAttributes.h
@@ -66,6 +66,12 @@ struct ENGINE_API FTimecodeCustomAttributeNameSettings
 	UPROPERTY(EditAnywhere, Category = TimecodeCustomAttributeNameSettings)
 	FName SubframeAttributeName;
 
+	/** Name of the custom attribute representing the timecode rate. This may be different from
+	    the animation or capture frame rate, for example when capturing "high" frame rate data
+		at 120 frames per second but recording SMPTE timecode at 30 frames per second. */
+	UPROPERTY(EditAnywhere, Category = TimecodeCustomAttributeNameSettings)
+	FName RateAttributeName;
+
 	/** Name of the custom attribute representing the name of a take. */
 	UPROPERTY(EditAnywhere, Category = TimecodeCustomAttributeNameSettings)
 	FName TakenameAttributeName;
diff --git a/Engine/Source/Runtime/Engine/Private/Animation/AnimationSettings.cpp b/Engine/Source/Runtime/Engine/Private/Animation/AnimationSettings.cpp
--- a/Engine/Source/Runtime/Engine/Private/Animation/AnimationSettings.cpp
+++ b/Engine/Source/Runtime/Engine/Private/Animation/AnimationSettings.cpp
@@ -42,6 +42,7 @@ TArray<FString> UAnimationSettings::GetBoneCustomAttributeNamesToImport() const
 		BoneTimecodeCustomAttributeNameSettings.SecondAttributeName.ToString(),
 		BoneTimecodeCustomAttributeNameSettings.FrameAttributeName.ToString(),
 		BoneTimecodeCustomAttributeNameSettings.SubframeAttributeName.ToString(),
+		BoneTimecodeCustomAttributeNameSettings.RateAttributeName.ToString(),
 		BoneTimecodeCustomAttributeNameSettings.TakenameAttributeName.ToString()
 	};
 
