AnimationBlueprintLibrary: add support for evaluating string-typed timecode custom attributes and timecode rates

This extends the handling in EvaluateRootBoneTimecodeAttributesAtTime() to support reading
numerical values from string-typed custom attributes for any of the existing components of
FTimecode or a subframe value.

It also adds support for parsing a timecode rate string into the FFrameRate it represents,
and accounts for the use of drop frame timecode rates (with either "29.97df" for NTSC_30
or "59.94df" for NTSC_60 frame rates).

diff --git a/Engine/Source/Editor/AnimationBlueprintLibrary/Private/AnimationBlueprintLibrary.cpp b/Engine/Source/Editor/AnimationBlueprintLibrary/Private/AnimationBlueprintLibrary.cpp
--- a/Engine/Source/Editor/AnimationBlueprintLibrary/Private/AnimationBlueprintLibrary.cpp
+++ b/Engine/Source/Editor/AnimationBlueprintLibrary/Private/AnimationBlueprintLibrary.cpp
@@ -2220,6 +2220,7 @@ bool UAnimationBlueprintLibrary::EvaluateRootBoneTimecodeAttributesAtTime(const
 	FTimecode Timecode;
 	float SubFrame = 0.0f;
 
+	FString TimecodeRateAsString;
 	double TimecodeRateAsDecimal = 0.0;
 
 	for (const FAnimatedBoneAttribute* RootBoneAttribute : RootBoneAttributes)
@@ -2244,6 +2245,7 @@ bool UAnimationBlueprintLibrary::EvaluateRootBoneTimecodeAttributesAtTime(const
 
 		float FloatValue = 0.0f;
 		int32 IntValue = 0;
+		FString StringValue;
 
 		// Support timecode attribute curves that are either float-typed or integer-typed.
 		if (RootBoneAttribute->Curve.GetScriptStruct() == FFloatAnimationAttribute::StaticStruct())
@@ -2258,6 +2260,16 @@ bool UAnimationBlueprintLibrary::EvaluateRootBoneTimecodeAttributesAtTime(const
 			IntValue = EvaluatedAttribute.Value;
 			FloatValue = static_cast<float>(IntValue);
 		}
+		else if (RootBoneAttribute->Curve.GetScriptStruct() == FStringAnimationAttribute::StaticStruct())
+		{
+			const FStringAnimationAttribute EvaluatedAttribute = RootBoneAttribute->Curve.Evaluate<FStringAnimationAttribute>(EvalTime);
+			StringValue = EvaluatedAttribute.Value;
+			if (StringValue.IsNumeric())
+			{
+				LexFromString(FloatValue, *StringValue);
+				IntValue = static_cast<int32>(FloatValue);
+			}
+		}
 		else
 		{
 			continue;
@@ -2290,6 +2302,7 @@ bool UAnimationBlueprintLibrary::EvaluateRootBoneTimecodeAttributesAtTime(const
 		}
 		else if (BoneAttributeName.IsEqual(TCRateAttrName))
 		{
+			TimecodeRateAsString = StringValue;
 			TimecodeRateAsDecimal = FloatValue;
 			// Don't consider this attribute when determining whether timecode value attributes are
 			// present, since it can't be useful on its own.
@@ -2305,7 +2318,11 @@ bool UAnimationBlueprintLibrary::EvaluateRootBoneTimecodeAttributesAtTime(const
 	// original source frame rate.
 	FFrameRate FrameRate = AnimationSequenceBase->GetSamplingFrameRate();
 
-	if (TimecodeRateAsDecimal > 1.0)
+	if (!TimecodeRateAsString.IsEmpty() && TryParseString(FrameRate, *TimecodeRateAsString))
+	{
+		Timecode.bDropFrameFormat = FTimecode::IsValidDropFormatTimecodeRate(TimecodeRateAsString);
+	}
+	else if (TimecodeRateAsDecimal > 1.0)
 	{
 		if (const FCommonFrameRateInfo* TimecodeRateInfo = FCommonFrameRates::Find(TimecodeRateAsDecimal))
 		{
