AnimationBlueprintLibrary: add function for evaluating the subframe timecode attribute

Since the subframe component of FFrameTime is clamped to the range [0.0, 1.0),
it cannot accurately represent the use case where the timecode metadata represents
subframe values as whole numbered subframes instead of as a percentage of a frame
the way the engine does. The subframe component of the FQualifiedFrameTime returned
by EvaluateRootBoneTimecodeAttributesAtTime() may not reflect the authored subframe
metadata in that case.

This new function allows access to the subframe values that were actually authored in the
timecode metadata.

An additional explicit clamp was added to EvaluateRootBoneTimecodeAttributesAtTime()
to ensure that the subframe value read from the timecode metadata is in range *before*
we try to construct an FFrameTime from it, since the constructor has a checkSlow() that
will assert if the value is out of range.

diff --git a/Engine/Source/Editor/AnimationBlueprintLibrary/Private/AnimationBlueprintLibrary.cpp b/Engine/Source/Editor/AnimationBlueprintLibrary/Private/AnimationBlueprintLibrary.cpp
--- a/Engine/Source/Editor/AnimationBlueprintLibrary/Private/AnimationBlueprintLibrary.cpp
+++ b/Engine/Source/Editor/AnimationBlueprintLibrary/Private/AnimationBlueprintLibrary.cpp
@@ -2341,6 +2341,18 @@ bool UAnimationBlueprintLibrary::EvaluateRootBoneTimecodeAttributesAtTime(const
 		}
 	}
 
+	// Some pipelines may author subframe values that are compatible with the
+	// engine's notion of a subframe, which is a floating point value between
+	// zero and one representing a percentage of a frame. Others may author
+	// subframe as integer values from zero to N instead, incrementing by one
+	// each subframe until the next frame is reached.
+	// Since this data is user-supplied, we don't want to trip over the
+	// checkSlow() in FFrameTime's constructor, so we apply the same clamping
+	// it does to bring the value into range here.
+	// Clients that are interested in the exact subframe value that was
+	// authored should query it using EvaluateRootBoneTimecodeSubframeAttributeAtTime().
+	SubFrame = FMath::Clamp(SubFrame + 0.5f - 0.5f, 0.f, FFrameTime::MaxSubframe);
+
 	OutQualifiedFrameTime = FQualifiedFrameTime(
 		FFrameTime(Timecode.ToFrameNumber(FrameRate), SubFrame),
 		FrameRate);
@@ -2348,6 +2360,52 @@ bool UAnimationBlueprintLibrary::EvaluateRootBoneTimecodeAttributesAtTime(const
 	return true;
 }
 
+bool UAnimationBlueprintLibrary::EvaluateRootBoneTimecodeSubframeAttributeAtTime(const UAnimSequenceBase* AnimationSequenceBase, const float EvalTime, float& OutSubframe)
+{
+	if (!AnimationSequenceBase)
+	{
+		return false;
+	}
+
+	const UAnimDataModel* AnimDataModel = AnimationSequenceBase->GetDataModel();
+	if (!AnimDataModel)
+	{
+		return false;
+	}
+
+	const int32 RootBoneIndex = 0;
+	const FBoneAnimationTrack* RootBoneTrack = AnimDataModel->FindBoneTrackByIndex(RootBoneIndex);
+	if (!RootBoneTrack)
+	{
+		return false;
+	}
+
+	const FName& RootBoneName = RootBoneTrack->Name;
+
+	FName TCSubframeAttrName(TEXT("TCSubframe"));
+	if (const UAnimationSettings* AnimationSettings = UAnimationSettings::Get())
+	{
+		TCSubframeAttrName = AnimationSettings->BoneTimecodeCustomAttributeNameSettings.SubframeAttributeName;
+	}
+
+	FAnimationAttributeIdentifier SubframeAttributeIdentifier(TCSubframeAttrName, RootBoneIndex, RootBoneName, FFloatAnimationAttribute::StaticStruct());
+	const FAnimatedBoneAttribute* RootBoneSubframeAttribute = AnimDataModel->FindAttribute(SubframeAttributeIdentifier);
+	if (!RootBoneSubframeAttribute)
+	{
+		return false;
+	}
+
+	if (!RootBoneSubframeAttribute->Curve.CanEvaluate())
+	{
+		return false;
+	}
+
+	const FFloatAnimationAttribute EvaluatedAttribute = RootBoneSubframeAttribute->Curve.Evaluate<FFloatAnimationAttribute>(EvalTime);
+	OutSubframe = EvaluatedAttribute.Value;
+
+	return true;
+}
+
 void UAnimationBlueprintLibrary::FindBonePathToRoot(const UAnimSequenceBase* AnimationSequenceBase, FName BoneName, TArray<FName>& BonePath)
 {
 	BonePath.Empty();
diff --git a/Engine/Source/Editor/AnimationBlueprintLibrary/Public/AnimationBlueprintLibrary.h b/Engine/Source/Editor/AnimationBlueprintLibrary/Public/AnimationBlueprintLibrary.h
--- a/Engine/Source/Editor/AnimationBlueprintLibrary/Public/AnimationBlueprintLibrary.h
+++ b/Engine/Source/Editor/AnimationBlueprintLibrary/Public/AnimationBlueprintLibrary.h
@@ -482,6 +482,24 @@ public:
 	UFUNCTION(BlueprintPure, Category = "AnimationBlueprintLibrary|Helpers")
 	static bool EvaluateRootBoneTimecodeAttributesAtTime(const UAnimSequenceBase* AnimationSequenceBase, const float EvalTime, FQualifiedFrameTime& OutQualifiedFrameTime);
 
+	/** Evaluates the subframe timecode attribute (e.g. "TCSubframe") of the root bone and returns the resulting value.
+	 *
+	 *  Since the subframe component of FFrameTime is clamped to the range [0.0, 1.0), it cannot accurately represent the use
+	 *  case where the timecode metadata represents subframe values as whole numbered subframes instead of as a percentage of a
+	 *  frame the way the engine does. The subframe component of the FQualifiedFrameTime returned by
+	 *  EvaluateRootBoneTimecodeAttributesAtTime() may not reflect the authored subframe metadata in that case.
+	 * 
+	 *  This function allows access to the subframe values that were actually authored in the timecode metadata.
+	 *
+	 *  @param AnimationSequenceBase: Anim sequence for which to evaluate the root bone subframe attribute.
+	 *  @param EvalTime: Time (in seconds) at which to evaluate the subframe timecode bone attribute.
+	 *  @param OutSubframe: Resulting subframe value from evaluation. If no subframe timecode attribute is present
+	 *      on the bone or if it cannot be evaluated, the output parameter will not be modified.
+	 *  @return: true if the root bone had a subframe timecode attribute that could be evaluated and a value was set, or false otherwise.
+	 */
+	UFUNCTION(BlueprintPure, Category = "AnimationBlueprintLibrary|Helpers")
+	static bool EvaluateRootBoneTimecodeSubframeAttributeAtTime(const UAnimSequenceBase* AnimationSequenceBase, const float EvalTime, float& OutSubframe);
+
 	/** Finds the Bone Path from the given Bone to the Root Bone */
 	UFUNCTION(BlueprintPure, Category = "AnimationBlueprintLibrary|Helpers")
 	static void FindBonePathToRoot(const UAnimSequenceBase* AnimationSequenceBase, FName BoneName, TArray<FName>& BonePath);
