Fixing crash when exporting animation in retargeter.

diff --git a/Engine/Source/Editor/AnimationBlueprintLibrary/Private/AnimPose.cpp b/Engine/Source/Editor/AnimationBlueprintLibrary/Private/AnimPose.cpp
--- a/Engine/Source/Editor/AnimationBlueprintLibrary/Private/AnimPose.cpp
+++ b/Engine/Source/Editor/AnimationBlueprintLibrary/Private/AnimPose.cpp
@@ -474,28 +474,30 @@ void UAnimPoseExtensions::GetAnimPoseAtTime(const UAnimSequenceBase* AnimationSe
 	if (AnimationSequenceBase && AnimationSequenceBase->GetSkeleton())
 	{
 		FMemMark Mark(FMemStack::Get());
-		const FReferenceSkeleton& RefSkeleton = AnimationSequenceBase->GetSkeleton()->GetReferenceSkeleton();
 
 		bool bValidTime = false;
 		UAnimationBlueprintLibrary::IsValidTime(AnimationSequenceBase, Time, bValidTime);
 		if (bValidTime)
-		{			       
-			TArray<FBoneIndexType> RequiredBoneIndexArray;
-			RequiredBoneIndexArray.AddUninitialized(RefSkeleton.GetNum());
-			for (int32 BoneIndex = 0; BoneIndex < RequiredBoneIndexArray.Num(); ++BoneIndex)
-			{
-				RequiredBoneIndexArray[BoneIndex] = BoneIndex;
-			}
-
+		{	
 			// asset to use for retarget proportions (can be either USkeletalMesh or USkeleton)
 			UObject* AssetToUse;
+			int32 NumRequiredBones;
 			if (EvaluationOptions.OptionalSkeletalMesh)
 			{
 				AssetToUse = CastChecked<UObject>(EvaluationOptions.OptionalSkeletalMesh);
+				NumRequiredBones = EvaluationOptions.OptionalSkeletalMesh->GetRefSkeleton().GetNum();	
 			}
 			else
 			{
 				AssetToUse = CastChecked<UObject>(AnimationSequenceBase->GetSkeleton());
+				NumRequiredBones = AnimationSequenceBase->GetSkeleton()->GetReferenceSkeleton().GetNum();
+			}
+
+			TArray<FBoneIndexType> RequiredBoneIndexArray;
+			RequiredBoneIndexArray.AddUninitialized(NumRequiredBones);
+			for (int32 BoneIndex = 0; BoneIndex < RequiredBoneIndexArray.Num(); ++BoneIndex)
+			{
+				RequiredBoneIndexArray[BoneIndex] = BoneIndex;
 			}
 
 			FBoneContainer RequiredBones;
