#feature added script exposed AddNodeAssetOverride function, allowing users to set Animation Asset overrides through Blueprints/Python

diff --git a/Engine/Source/Editor/AnimationBlueprintLibrary/AnimationBlueprintLibrary.build.cs b/Engine/Source/Editor/AnimationBlueprintLibrary/AnimationBlueprintLibrary.build.cs
--- a/Engine/Source/Editor/AnimationBlueprintLibrary/AnimationBlueprintLibrary.build.cs
+++ b/Engine/Source/Editor/AnimationBlueprintLibrary/AnimationBlueprintLibrary.build.cs
@@ -14,6 +14,7 @@ public class AnimationBlueprintLibrary : ModuleRules
 				"Engine",
 				"Kismet",
 				"AnimGraph",
+				"UnrealEd",
 				"TimeManagement"
 			}
 		);
diff --git a/Engine/Source/Editor/AnimationBlueprintLibrary/Private/AnimationBlueprintLibrary.cpp b/Engine/Source/Editor/AnimationBlueprintLibrary/Private/AnimationBlueprintLibrary.cpp
--- a/Engine/Source/Editor/AnimationBlueprintLibrary/Private/AnimationBlueprintLibrary.cpp
+++ b/Engine/Source/Editor/AnimationBlueprintLibrary/Private/AnimationBlueprintLibrary.cpp
@@ -26,6 +26,8 @@
 #include "AnimationRuntime.h"
 #include "Animation/AnimData/AnimDataModel.h"
 #include "Animation/AnimSequenceHelpers.h"
+#include "AnimGraphNode_AssetPlayerBase.h"
+#include "Kismet2/BlueprintEditorUtils.h"
 
 #define LOCTEXT_NAMESPACE "AnimationBlueprintLibrary"
 
@@ -2548,6 +2550,130 @@ void UAnimationBlueprintLibrary::GetNodesOfClass(UAnimBlueprint* AnimationBluepr
 	}
 }
 
+void UAnimationBlueprintLibrary::AddNodeAssetOverride(UAnimBlueprint* AnimBlueprint, const UAnimationAsset* Target, UAnimationAsset* Override, bool bPrintAppliedOverrides /*= false*/)
+{
+	if (AnimBlueprint)
+	{
+		if (Target && Override)
+		{
+			// Target and override animation asset types have to match
+			if (Target->GetClass() == Override->GetClass())			
+			{
+				TArray<UBlueprint*> BlueprintHierarchy;
+				AnimBlueprint->GetBlueprintHierarchyFromClass(AnimBlueprint->GetAnimBlueprintGeneratedClass(), BlueprintHierarchy);
+
+				TArray<const UAnimGraphNode_Base*> OverridableNodes;
+
+				// Search from 1 as 0 is this class and we're looking for it's parents
+				for(int32 BlueprintIndex = 1; BlueprintIndex < BlueprintHierarchy.Num(); ++BlueprintIndex)
+				{
+					const UBlueprint* CurrentBlueprint = BlueprintHierarchy[BlueprintIndex];
+
+					TArray<UEdGraph*> Graphs;
+					CurrentBlueprint->GetAllGraphs(Graphs);
+
+					for(const UEdGraph* Graph : Graphs)
+					{
+						for(const UEdGraphNode* Node : Graph->Nodes)
+						{
+							const UAnimGraphNode_Base* AnimNode = Cast<UAnimGraphNode_Base>(Node);
+							// Find any overridable node with Target set as its current value
+							if(AnimNode && AnimNode->GetAnimationAsset() == Target)
+							{
+								OverridableNodes.Add(AnimNode);
+							}
+						}
+					}
+				}
+
+				// Apply overrides
+				for (const UAnimGraphNode_Base* OverrideNode : OverridableNodes)
+				{
+					const FGuid NodeGUID = OverrideNode->NodeGuid;
+					
+					FAnimParentNodeAssetOverride* OverridePtr = AnimBlueprint->ParentAssetOverrides.FindByPredicate([NodeGUID](const FAnimParentNodeAssetOverride& Other)
+					{
+						return Other.ParentNodeGuid == NodeGUID;
+					});
+						
+					if (OverridePtr == nullptr)
+					{
+						OverridePtr = &AnimBlueprint->ParentAssetOverrides.AddDefaulted_GetRef();
+					}
+
+					check(OverridePtr != nullptr);
+
+					auto GetOverrideNodeTitle = [OverrideNode]() -> FString
+					{
+						if (const UAnimGraphNode_AssetPlayerBase * AssetPlayerBase = Cast<UAnimGraphNode_AssetPlayerBase>(OverrideNode))
+						{
+							return AssetPlayerBase->GetNodeTitle(ENodeTitleType::FullTitle).ToString();
+						}
+
+						return OverrideNode->GetName();
+					};
+
+					if (OverridePtr->NewAsset != Override)
+					{
+						// Setup override values
+						OverridePtr->NewAsset = Override;
+						OverridePtr->ParentNodeGuid = NodeGUID;
+					
+						AnimBlueprint->NotifyOverrideChange(*OverridePtr);
+
+						if (bPrintAppliedOverrides)
+						{
+							UE_LOG(LogAnimationBlueprintLibrary, Display, TEXT("Set Animation Blueprint asset override for %s\n\tAnimation Node: %s\n\tAnimation Node Bueprint: %s\n\tOriginal: %s\n\tOverride: %s"),
+								*AnimBlueprint->GetPathName(),
+								*GetOverrideNodeTitle(),
+								*OverrideNode->GetAnimBlueprint()->GetPathName(),
+								*Target->GetPathName(),
+								*Override->GetPathName()							
+							);
+						}
+					}
+					else
+					{
+						
+						UE_LOG(LogAnimationBlueprintLibrary, Warning, TEXT("Found matching pre-existing Animation Blueprint asset override for %s\n\tAnimation Node: %s\n\tAnimation Node Bueprint: %s\n\tOriginal: %s\n\tOverride: %s"),
+								*AnimBlueprint->GetPathName(),
+								*GetOverrideNodeTitle(),
+								*OverrideNode->GetAnimBlueprint()->GetPathName(),
+								*Target->GetPathName(),
+								*Override->GetPathName()							
+							);
+					}
+				}
+
+				if (OverridableNodes.Num())
+				{					
+					FBlueprintEditorUtils::MarkBlueprintAsModified(AnimBlueprint);
+				}
+			}
+			else
+			{
+				UE_LOG(LogAnimationBlueprintLibrary, Error, TEXT("Failed to add override as Target and Override class do not match, expected %s but found %s"), *Target->GetClass()->GetName(), *Override->GetClass()->GetName());
+			}
+		}
+		else
+		{
+			if (Target == nullptr)
+			{
+				UE_LOG(LogAnimationBlueprintLibrary, Error, TEXT("Failed to add override as provided Target asset is invalid"));
+			}
+			
+			if (Override == nullptr)
+			{
+				UE_LOG(LogAnimationBlueprintLibrary, Error, TEXT("Failed to add override as provided Override asset is invalid"));
+			}
+		}
+	}
+	else
+	{
+		UE_LOG(LogAnimationBlueprintLibrary, Error, TEXT("Failed to add override as provided Animation Blueprint is invalid"));
+	}
+}
+
 template void UAnimationBlueprintLibrary::AddCurveKeysInternal<float, FFloatCurve, ERawCurveTrackTypes::RCT_Float>(UAnimSequence* AnimationSequence, FName CurveName, const TArray<float>& Times, const TArray<float>& KeyData);
 template void UAnimationBlueprintLibrary::AddCurveKeysInternal<FVector, FVectorCurve, ERawCurveTrackTypes::RCT_Vector>(UAnimSequence* AnimationSequence, FName CurveName, const TArray<float>& Times, const TArray<FVector>& KeyData);
 template void UAnimationBlueprintLibrary::AddCurveKeysInternal<FTransform, FTransformCurve, ERawCurveTrackTypes::RCT_Transform>(UAnimSequence* AnimationSequence, FName CurveName, const TArray<float>& Times, const TArray<FTransform>& KeyData);
diff --git a/Engine/Source/Editor/AnimationBlueprintLibrary/Public/AnimationBlueprintLibrary.h b/Engine/Source/Editor/AnimationBlueprintLibrary/Public/AnimationBlueprintLibrary.h
--- a/Engine/Source/Editor/AnimationBlueprintLibrary/Public/AnimationBlueprintLibrary.h
+++ b/Engine/Source/Editor/AnimationBlueprintLibrary/Public/AnimationBlueprintLibrary.h
@@ -511,4 +511,15 @@ public:
 	/** Returns all Animation Graph Nodes of the provided Node Class contained by the Animation Blueprint */
 	UFUNCTION(BlueprintCallable, Category=Animation, meta=(ScriptMethod))
 	static void GetNodesOfClass(UAnimBlueprint* AnimationBlueprint, TSubclassOf<UAnimGraphNode_Base> NodeClass, TArray<UAnimGraphNode_Base*>& GraphNodes, bool bIncludeChildClasses = true);
+
+	/**
+	 * Adds an Animation Asset override for the provided AnimationBlueprint, replacing any instance of Target with Override
+	 *
+	 * @param AnimBlueprint					The Animation Blueprint to add/set the Override for
+	 * @param Target						The Animation Asset to add an override for (overrides all instances of the asset)
+	 * @param Override						The Animation Asset to used to override the Target with (types have to match)
+	 * @param bPrintAppliedOverrides		Flag whether or not to print the applied overrides
+	 */
+	UFUNCTION(BlueprintCallable, Category=Animation, meta = (ScriptMethod))
+	static void AddNodeAssetOverride(UAnimBlueprint* AnimBlueprint, const UAnimationAsset* Target, UAnimationAsset* Override, bool bPrintAppliedOverrides = false);
 };
