Fix Linux build due to use of inline function GetAnimBlueprint
#rnx

diff --git a/Engine/Source/Editor/AnimationBlueprintLibrary/AnimationBlueprintLibrary.build.cs b/Engine/Source/Editor/AnimationBlueprintLibrary/AnimationBlueprintLibrary.build.cs
--- a/Engine/Source/Editor/AnimationBlueprintLibrary/AnimationBlueprintLibrary.build.cs
+++ b/Engine/Source/Editor/AnimationBlueprintLibrary/AnimationBlueprintLibrary.build.cs
@@ -9,6 +9,7 @@ public class AnimationBlueprintLibrary : ModuleRules
 		PrivateDependencyModuleNames.AddRange(
 			new string[]
 			{
+				"BlueprintGraph",
 				"Core",
 				"CoreUObject",
 				"Engine",
