Added uniform blueprint nodes for adding/updating of Animation Attribute (data)
- Expose built in attribute type properties to blueprints/scripting

diff --git a/Engine/Source/Editor/AnimationBlueprintLibrary/AnimationBlueprintLibrary.build.cs b/Engine/Source/Editor/AnimationBlueprintLibrary/AnimationBlueprintLibrary.build.cs
--- a/Engine/Source/Editor/AnimationBlueprintLibrary/AnimationBlueprintLibrary.build.cs
+++ b/Engine/Source/Editor/AnimationBlueprintLibrary/AnimationBlueprintLibrary.build.cs
@@ -12,7 +12,6 @@ public class AnimationBlueprintLibrary : ModuleRules
 				"BlueprintGraph",
 				"Core",
 				"CoreUObject",
-				"Engine",
 				"Kismet",
 				"AnimGraph",
 				"UnrealEd",
@@ -24,6 +23,10 @@ public class AnimationBlueprintLibrary : ModuleRules
 			new string[]
 			{
 				"AnimGraph",
+				"KismetCompiler",
+				"Engine",
+				"BlueprintGraph",
+				"UnrealEd",
 			}
 		);
 	}
diff --git a/Engine/Source/Editor/AnimationBlueprintLibrary/Private/AnimationAttributeBlueprintLibrary.cpp b/Engine/Source/Editor/AnimationBlueprintLibrary/Private/AnimationAttributeBlueprintLibrary.cpp
new file mode 100644
--- /dev/null
+++ b/Engine/Source/Editor/AnimationBlueprintLibrary/Private/AnimationAttributeBlueprintLibrary.cpp
@@ -0,0 +1,294 @@
+// Copyright Epic Games, Inc. All Rights Reserved.
+
+#include "AnimationAttributeBlueprintLibrary.h"
+
+#define LOCTEXT_NAMESPACE "AnimationAttributeBlueprintLibrary"
+
+bool UAnimationAttributeBlueprintLibrary::SetAttributeKey(TScriptInterface<IAnimationDataController> AnimationDataController, const FAnimationAttributeIdentifier& AttributeIdentifier, float Time, const int32& Value)
+{
+	// We should never hit this! Stubbed to avoid NoExport on the class.
+	check(0);
+	return false;
+}
+
+bool UAnimationAttributeBlueprintLibrary::Generic_SetAttributeKey(TScriptInterface<IAnimationDataController> AnimationDataController, const FAnimationAttributeIdentifier& AttributeIdentifier, float Time, UScriptStruct* ScriptStruct, const void* ValuePtr)
+{
+	return AnimationDataController->SetAttributeKey(AttributeIdentifier, Time, ValuePtr, ScriptStruct);
+}
+
+DEFINE_FUNCTION(UAnimationAttributeBlueprintLibrary::execSetAttributeKey)
+{
+	P_GET_TINTERFACE(IAnimationDataController, AnimationDataController);
+	P_GET_STRUCT_REF(FAnimationAttributeIdentifier, AttributeIdentifier);
+	P_GET_PROPERTY(FFloatProperty, TimeInterval);
+
+	Stack.MostRecentProperty = nullptr;
+	Stack.MostRecentPropertyAddress = nullptr;
+	
+	Stack.StepCompiledIn<FStructProperty>(nullptr);
+	
+	const FStructProperty* ItemProperty = CastField<FStructProperty>(Stack.MostRecentProperty);
+	void* ItemDataPtr = Stack.MostRecentPropertyAddress;
+	P_FINISH;
+
+	if (!ItemProperty || !ItemDataPtr)
+	{
+		FBlueprintExceptionInfo ExceptionInfo(
+			EBlueprintExceptionType::AccessViolation,
+			LOCTEXT("SetAttributeKey_InvalidValue", "Failed to resolve the attribute value parameter for SetAttributeKey.")
+		);
+		FBlueprintCoreDelegates::ThrowScriptException(P_THIS, Stack, ExceptionInfo);
+	}
+
+	if (!AnimationDataController.GetObject())
+	{
+		FBlueprintExceptionInfo ExceptionInfo(
+			EBlueprintExceptionType::AccessViolation, 
+			LOCTEXT("SetAttributeKey_InvalidController", "Accessed None attempting to call SetAttributeKey.")
+		);
+		FBlueprintCoreDelegates::ThrowScriptException(P_THIS, Stack, ExceptionInfo);
+	}
+
+	bool bResult = false;
+
+	if (AnimationDataController && ItemProperty)
+	{		
+		P_NATIVE_BEGIN;
+		bResult = Generic_SetAttributeKey(AnimationDataController, AttributeIdentifier, TimeInterval, ItemProperty->Struct, ItemDataPtr);
+		P_NATIVE_END;
+	}
+
+	*(bool*)RESULT_PARAM = bResult;
+}
+
+bool UAnimationAttributeBlueprintLibrary::SetAttributeKeys(TScriptInterface<IAnimationDataController> AnimationDataController, const FAnimationAttributeIdentifier& AttributeIdentifier, const TArray<float>& Times, const TArray<int32>& Values)
+{
+	// We should never hit this! Stubbed to avoid NoExport on the class.
+	check(0);
+	return false;
+}
+
+bool UAnimationAttributeBlueprintLibrary::Generic_SetAttributeKeys(TScriptInterface<IAnimationDataController> AnimationDataController, const FAnimationAttributeIdentifier& AttributeIdentifier, const TArray<float>& Times, const void* ValuesArray, const FArrayProperty* ValuesArrayProperty)
+{	
+	FScriptArrayHelper ArrayHelper(ValuesArrayProperty, ValuesArray);
+
+	TArray<const void*> ItemPtrs;
+	const int32 NumberOfItems = ArrayHelper.Num();
+	ItemPtrs.SetNumZeroed(NumberOfItems);
+	for (int32 Index = 0; Index < NumberOfItems; ++Index)
+	{
+		ItemPtrs[Index] = ArrayHelper.GetRawPtr(Index);
+	}
+
+	const FStructProperty* InnerStructProperty = CastField<const FStructProperty>(ValuesArrayProperty->Inner);		
+	return AnimationDataController->SetAttributeKeys(AttributeIdentifier, MakeArrayView(Times), MakeArrayView(ItemPtrs), InnerStructProperty->Struct);
+}
+
+DEFINE_FUNCTION(UAnimationAttributeBlueprintLibrary::execSetAttributeKeys)
+{
+	P_GET_TINTERFACE(IAnimationDataController, AnimationDataController);
+	P_GET_STRUCT_REF(FAnimationAttributeIdentifier, AttributeIdentifier);
+	P_GET_TARRAY_REF(float, TimeIntervals);
+
+	Stack.MostRecentProperty = nullptr;
+	Stack.MostRecentPropertyAddress = nullptr;
+	
+	Stack.StepCompiledIn<FArrayProperty>(nullptr);
+
+	void* ArrayAddr = Stack.MostRecentPropertyAddress;
+	FArrayProperty* ArrayProperty = CastField<FArrayProperty>(Stack.MostRecentProperty);	
+	P_FINISH;
+
+	if (!ArrayAddr || !ArrayProperty)
+	{
+		FBlueprintExceptionInfo ExceptionInfo(
+			EBlueprintExceptionType::AccessViolation,
+			LOCTEXT("SetAttributeKeys_InvalidValue", "Failed to resolve the attribute values parameter for SetAttributeKeys.")
+		);
+		FBlueprintCoreDelegates::ThrowScriptException(P_THIS, Stack, ExceptionInfo);
+	}
+	
+	if (!AnimationDataController.GetObject())
+	{
+		FBlueprintExceptionInfo ExceptionInfo(
+			EBlueprintExceptionType::AccessViolation, 
+			LOCTEXT("SetAttributeKey_InvalidController", "Accessed None attempting to call SetAttributeKeys.")
+		);
+		FBlueprintCoreDelegates::ThrowScriptException(P_THIS, Stack, ExceptionInfo);
+	}
+
+	bool bResult = false;
+
+	if (AnimationDataController)
+	{		
+		P_NATIVE_BEGIN;
+		bResult = Generic_SetAttributeKeys(AnimationDataController, AttributeIdentifier, TimeIntervals, ArrayAddr, ArrayProperty);
+		P_NATIVE_END;
+	}
+
+	*(bool*)RESULT_PARAM = bResult;
+}
+
+
+bool UAnimationAttributeBlueprintLibrary::GetAttributeKey(UAnimDataModel* AnimationDataModel, const FAnimationAttributeIdentifier& AttributeIdentifier, float Time, int32& Value)
+{
+	// We should never hit this! Stubbed to avoid NoExport on the class.
+	check(0);
+	return false;
+}
+
+bool UAnimationAttributeBlueprintLibrary::Generic_GetAttributeKey(UAnimDataModel* AnimationDataModel, const FAnimationAttributeIdentifier& AttributeIdentifier, float Time, UScriptStruct* ScriptStruct, void* ValuePtr)
+{
+	const FAnimatedBoneAttribute* AttributePtr = AnimationDataModel->FindAttribute(AttributeIdentifier);
+
+	// Find attribute
+	if (AttributePtr && ScriptStruct)
+	{
+		// And the key on its attribute curve
+		const FAttributeCurve& Curve = AttributePtr->Curve;
+		check(ScriptStruct == Curve.GetScriptStruct());
+		const FKeyHandle KeyHandle = Curve.FindKey(Time);
+		if (KeyHandle != FKeyHandle::Invalid())
+		{
+			// Copy out the value to ptr
+			const FAttributeKey& Key = Curve.GetKey(KeyHandle);
+			ScriptStruct->CopyScriptStruct(ValuePtr, Key.GetValuePtr<void>(), 1);
+			
+			return true;
+		}
+	}
+
+	return false;
+}
+
+DEFINE_FUNCTION(UAnimationAttributeBlueprintLibrary::execGetAttributeKey)
+{
+	P_GET_OBJECT(UAnimDataModel, AnimationDataModel);
+	P_GET_STRUCT_REF(FAnimationAttributeIdentifier, AttributeIdentifier);
+	P_GET_PROPERTY(FFloatProperty, TimeInterval);
+
+	Stack.MostRecentProperty = nullptr;
+	Stack.MostRecentPropertyAddress = nullptr;
+	
+	Stack.StepCompiledIn<FStructProperty>(nullptr);
+	
+	const FStructProperty* ItemProperty = CastField<FStructProperty>(Stack.MostRecentProperty);
+	void* ItemDataPtr = Stack.MostRecentPropertyAddress;
+	P_FINISH;
+
+	if (!ItemProperty || !ItemDataPtr)
+	{
+		FBlueprintExceptionInfo ExceptionInfo(
+			EBlueprintExceptionType::AccessViolation,
+			LOCTEXT("SetAttributeKey_InvalidValue", "Failed to resolve the attribute value parameter for GetAttributeKey.")
+		);
+		FBlueprintCoreDelegates::ThrowScriptException(P_THIS, Stack, ExceptionInfo);
+	}
+
+	if (!AnimationDataModel)
+	{
+		FBlueprintExceptionInfo ExceptionInfo(
+			EBlueprintExceptionType::AccessViolation, 
+			LOCTEXT("SetAttributeKey_InvalidController", "Accessed None attempting to call GetAttributeKey.")
+		);
+		FBlueprintCoreDelegates::ThrowScriptException(P_THIS, Stack, ExceptionInfo);
+	}
+
+	bool bResult = false;
+
+	if (AnimationDataModel && ItemProperty)
+	{		
+		P_NATIVE_BEGIN;
+		bResult = Generic_GetAttributeKey(AnimationDataModel, AttributeIdentifier, TimeInterval, ItemProperty->Struct, ItemDataPtr);
+		P_NATIVE_END;
+	}
+
+	*(bool*)RESULT_PARAM = bResult;
+}
+
+
+bool UAnimationAttributeBlueprintLibrary::GetAttributeKeys(UAnimDataModel* AnimationDataModel, const FAnimationAttributeIdentifier& AttributeIdentifier, TArray<float>& OutTimes, TArray<int32>& OutValues)
+{
+	// We should never hit this! Stubbed to avoid NoExport on the class.
+	check(0);
+	return false;
+}
+
+bool UAnimationAttributeBlueprintLibrary::Generic_GetAttributeKeys(UAnimDataModel* AnimationDataModel, const FAnimationAttributeIdentifier& AttributeIdentifier, TArray<float>& Times, void* ValuesArray, const FArrayProperty* ValuesArrayProperty)
+{
+	const FAnimatedBoneAttribute* AttributePtr = AnimationDataModel->FindAttribute(AttributeIdentifier);
+
+	// Find attribute
+	if (AttributePtr && ValuesArrayProperty)
+	{
+		const FStructProperty* InnerStructProperty = CastField<const FStructProperty>(ValuesArrayProperty->Inner);	
+		// And the key on its attribute curve
+		const FAttributeCurve& Curve = AttributePtr->Curve;
+		check(InnerStructProperty->Struct == Curve.GetScriptStruct());
+
+		const TArray<FAttributeKey>& Keys = Curve.GetConstRefOfKeys();
+
+		
+		FScriptArrayHelper ArrayHelper(ValuesArrayProperty, ValuesArray);
+		ArrayHelper.Resize(Keys.Num());
+		Times.SetNum(Keys.Num());
+		
+		for (int32 KeyIndex = 0; KeyIndex < Keys.Num(); ++KeyIndex)
+		{
+			const FAttributeKey& Key = Keys[KeyIndex];
+			Times[KeyIndex] = Key.Time;
+
+			uint8* ValuePtr = ArrayHelper.GetRawPtr(KeyIndex);
+			InnerStructProperty->Struct->CopyScriptStruct(ValuePtr, Key.GetValuePtr<void>(), 1);
+		}
+	}
+
+	return false;
+}
+
+DEFINE_FUNCTION(UAnimationAttributeBlueprintLibrary::execGetAttributeKeys)
+{
+	P_GET_OBJECT(UAnimDataModel, AnimationDataModel);
+	P_GET_STRUCT_REF(FAnimationAttributeIdentifier, AttributeIdentifier);
+	P_GET_TARRAY_REF(float, TimeIntervals);
+
+	Stack.MostRecentProperty = nullptr;
+	Stack.MostRecentPropertyAddress = nullptr;
+	
+	Stack.StepCompiledIn<FArrayProperty>(nullptr);
+
+	void* ArrayAddr = Stack.MostRecentPropertyAddress;
+	FArrayProperty* ArrayProperty = CastField<FArrayProperty>(Stack.MostRecentProperty);	
+	P_FINISH;
+
+	if (!ArrayAddr || !ArrayProperty)
+	{
+		FBlueprintExceptionInfo ExceptionInfo(
+			EBlueprintExceptionType::AccessViolation,
+			LOCTEXT("SetAttributeKeys_InvalidValue", "Failed to resolve the attribute values parameter for GetAttributeKeys.")
+		);
+		FBlueprintCoreDelegates::ThrowScriptException(P_THIS, Stack, ExceptionInfo);
+	}
+	
+	if (!AnimationDataModel)
+	{
+		FBlueprintExceptionInfo ExceptionInfo(
+			EBlueprintExceptionType::AccessViolation, 
+			LOCTEXT("SetAttributeKey_InvalidController", "Accessed None attempting to call GetAttributeKeys.")
+		);
+		FBlueprintCoreDelegates::ThrowScriptException(P_THIS, Stack, ExceptionInfo);
+	}
+
+	bool bResult = false;
+
+	if (AnimationDataModel)
+	{		
+		P_NATIVE_BEGIN;
+		bResult = Generic_GetAttributeKeys(AnimationDataModel, AttributeIdentifier, TimeIntervals, ArrayAddr, ArrayProperty);
+		P_NATIVE_END;
+	}
+
+	*(bool*)RESULT_PARAM = bResult;
+}
+
+#undef LOCTEXT_NAMESPACE
diff --git a/Engine/Source/Editor/AnimationBlueprintLibrary/Private/K2_AnimationAttributeNodes.cpp b/Engine/Source/Editor/AnimationBlueprintLibrary/Private/K2_AnimationAttributeNodes.cpp
new file mode 100644
--- /dev/null
+++ b/Engine/Source/Editor/AnimationBlueprintLibrary/Private/K2_AnimationAttributeNodes.cpp
@@ -0,0 +1,111 @@
+// Copyright Epic Games, Inc. All Rights Reserved.
+
+#include "K2_AnimationAttributeNodes.h"
+
+#include "AnimationAttributeBlueprintLibrary.h"
+#include "BlueprintNodeSpawner.h"
+#include "KismetCompiler.h"
+#include "EdGraphSchema_K2.h"
+#include "Kismet2/BlueprintEditorUtils.h"
+#include "K2Node_CallFunction.h"
+#include "BlueprintActionDatabaseRegistrar.h"
+#include "Animation/AttributeTypes.h"
+
+#define LOCTEXT_NAMESPACE "K2Node_AnimationAttributeNodes"
+
+FName UK2Node_BaseAttributeActionNode::ValuePinName(TEXT("Value"));
+FName UK2Node_BaseAttributeActionNode::ValuesPinName(TEXT("Values"));
+
+FText UK2Node_BaseAttributeActionNode::GetNodeTitle(ENodeTitleType::Type Title) const
+{
+	ensure(AttributeValuePinName != NAME_None);
+	ensure(!AttributeActionFormat.IsEmpty());
+	
+	UEdGraphPin* Pin = FindPin(AttributeValuePinName);
+	const UScriptStruct* AttributeTypeStruct = Pin ? Cast<UScriptStruct>(Pin->PinType.PinSubCategoryObject.Get()) : nullptr;
+	return FText::Format(AttributeActionFormat, AttributeTypeStruct ? FText::FromString(AttributeTypeStruct->GetName()) : LOCTEXT("AttributeSub", "Attribute"));	
+}
+
+void UK2Node_BaseAttributeActionNode::GetMenuActions(FBlueprintActionDatabaseRegistrar& ActionRegistrar) const
+{
+	ensure(AttributeValuePinName != NAME_None);
+	
+	const UClass* ActionKey = GetClass();
+	if (ActionRegistrar.IsOpenForRegistration(ActionKey))
+	{
+		UBlueprintNodeSpawner* NodeSpawner = UBlueprintNodeSpawner::Create(GetClass());
+		check(NodeSpawner != nullptr);
+
+		ActionRegistrar.AddBlueprintAction(ActionKey, NodeSpawner);
+	}
+}
+
+bool UK2Node_BaseAttributeActionNode::IsConnectionDisallowed(const UEdGraphPin* MyPin, const UEdGraphPin* OtherPin, FString& OutReason) const
+{
+	ensure(AttributeValuePinName != NAME_None);
+	
+	UEdGraphPin* Pin = FindPinChecked(AttributeValuePinName);
+	if (Pin && MyPin == Pin)
+	{
+		bool bAllowed = false;
+		// Ensure the pin is a struct, and in an array pin in case it is expected/required to be
+		if (OtherPin->PinType.PinCategory == UEdGraphSchema_K2::PC_Struct)
+		{
+			if (AttributeValuePinName == ValuesPinName && !OtherPin->PinType.IsArray())
+			{
+				OutReason = LOCTEXT("NotArrayTypePin", "Must be an array of registered Animation Attributes.").ToString();
+			}
+			else if (AttributeValuePinName == ValuePinName && OtherPin->PinType.IsContainer())
+			{
+				OutReason = LOCTEXT("ContainerTypePin","Must not be a container pin type.").ToString();
+			}
+			else if (const UScriptStruct* AttributeType = Cast<UScriptStruct>(OtherPin->PinType.PinSubCategoryObject.Get()))
+			{
+				bAllowed = UE::Anim::AttributeTypes::IsTypeRegistered(AttributeType);
+
+				if (!bAllowed)
+				{
+					OutReason = FText::Format(LOCTEXT("UnregisterdStructureTypeFormat", "Must be a registered Animation Attribute, which {0} is not."), FText::FromName(AttributeType->GetFName())).ToString();
+				}
+			}
+		}
+		else
+		{
+			OutReason =LOCTEXT("UnregisterdStructureType","Must be a registered Animation Attribute.").ToString();
+		}
+
+		return !bAllowed;
+	}
+	
+	return Super::IsConnectionDisallowed(MyPin, OtherPin, OutReason);
+}
+
+UK2Node_SetAttributeKeyGeneric::UK2Node_SetAttributeKeyGeneric()
+{
+	FunctionReference.SetExternalMember(GET_FUNCTION_NAME_CHECKED(UAnimationAttributeBlueprintLibrary, SetAttributeKey), UAnimationAttributeBlueprintLibrary::StaticClass());
+	AttributeActionFormat = LOCTEXT("SetAttributeKeyNodetitle", "Set {0} Key");
+	AttributeValuePinName = ValuePinName;
+}
+
+UK2Node_SetAttributeKeysGeneric::UK2Node_SetAttributeKeysGeneric()
+{
+	FunctionReference.SetExternalMember(GET_FUNCTION_NAME_CHECKED(UAnimationAttributeBlueprintLibrary, SetAttributeKeys), UAnimationAttributeBlueprintLibrary::StaticClass());
+	AttributeActionFormat = LOCTEXT("SetAttributeKesyNodetitle", "Set {0} Keys");
+	AttributeValuePinName = ValuesPinName;
+}
+
+UK2Node_GetAttributeKeyGeneric::UK2Node_GetAttributeKeyGeneric()
+{
+	FunctionReference.SetExternalMember(GET_FUNCTION_NAME_CHECKED(UAnimationAttributeBlueprintLibrary, GetAttributeKey), UAnimationAttributeBlueprintLibrary::StaticClass());
+	AttributeActionFormat = LOCTEXT("GetAttributeKeyNodetitle", "Get {0} Key");
+	AttributeValuePinName = ValuePinName;
+}
+
+UK2Node_GetAttributeKeysGeneric::UK2Node_GetAttributeKeysGeneric()
+{
+	FunctionReference.SetExternalMember(GET_FUNCTION_NAME_CHECKED(UAnimationAttributeBlueprintLibrary, GetAttributeKeys), UAnimationAttributeBlueprintLibrary::StaticClass());
+	AttributeActionFormat = LOCTEXT("GetAttributeKeysNodetitle", "Get {0} Keys");
+	AttributeValuePinName = ValuesPinName;
+}
+
+#undef LOCTEXT_NAMESPACE
diff --git a/Engine/Source/Editor/AnimationBlueprintLibrary/Public/AnimationAttributeBlueprintLibrary.h b/Engine/Source/Editor/AnimationBlueprintLibrary/Public/AnimationAttributeBlueprintLibrary.h
new file mode 100644
--- /dev/null
+++ b/Engine/Source/Editor/AnimationBlueprintLibrary/Public/AnimationAttributeBlueprintLibrary.h
@@ -0,0 +1,37 @@
+// Copyright Epic Games, Inc. All Rights Reserved.
+
+#pragma once
+
+#include "Kismet/BlueprintFunctionLibrary.h"
+#include "Animation/BuiltInAttributeTypes.h"
+#include "Animation/AnimData/IAnimationDataController.h"
+
+#include "AnimationAttributeBlueprintLibrary.generated.h"
+
+UCLASS(meta=(ScriptName="AnimationAttributeLibrary"))
+class ANIMATIONBLUEPRINTLIBRARY_API UAnimationAttributeBlueprintLibrary : public UBlueprintFunctionLibrary
+{
+	GENERATED_BODY()
+
+public:
+	UFUNCTION(BlueprintCallable, CustomThunk, Category = "Animation|Attributes", meta=(CustomStructureParam="Value", BlueprintInternalUseOnly="true"))
+	static bool SetAttributeKey(TScriptInterface<IAnimationDataController> AnimationDataController, const FAnimationAttributeIdentifier& AttributeIdentifier, float Time, const int32& Value);
+	static bool Generic_SetAttributeKey(TScriptInterface<IAnimationDataController> AnimationDataController, const FAnimationAttributeIdentifier& AttributeIdentifier, float Time, UScriptStruct* ScriptStruct, const void* Value);
+	DECLARE_FUNCTION(execSetAttributeKey);
+
+	UFUNCTION(BlueprintCallable, CustomThunk, Category = "Animation|Attributes", meta=(CustomStructureParam="Values", BlueprintInternalUseOnly="true"))
+	static bool SetAttributeKeys(TScriptInterface<IAnimationDataController> AnimationDataController, const FAnimationAttributeIdentifier& AttributeIdentifier, const TArray<float>& Times, const TArray<int32>& Values);
+	static bool Generic_SetAttributeKeys(TScriptInterface<IAnimationDataController> AnimationDataController, const FAnimationAttributeIdentifier& AttributeIdentifier, const TArray<float>& Times, const void* ValuesArray, const FArrayProperty* ValuesArrayProperty);
+	DECLARE_FUNCTION(execSetAttributeKeys);
+
+
+	UFUNCTION(BlueprintCallable, CustomThunk, Category = "Animation|Attributes", meta=(CustomStructureParam="Value", BlueprintInternalUseOnly="true"))
+	static bool GetAttributeKey(UAnimDataModel* AnimationDataModel, const FAnimationAttributeIdentifier& AttributeIdentifier, float Time, int32& Value);
+	static bool Generic_GetAttributeKey(UAnimDataModel* AnimationDataModel, const FAnimationAttributeIdentifier& AttributeIdentifier, float Time, UScriptStruct* ScriptStruct, void* Value);
+	DECLARE_FUNCTION(execGetAttributeKey);
+
+	UFUNCTION(BlueprintCallable, CustomThunk, Category = "Animation|Attributes", meta=(CustomStructureParam="Values", BlueprintInternalUseOnly="true"))
+	static bool GetAttributeKeys(UAnimDataModel* AnimationDataModel, const FAnimationAttributeIdentifier& AttributeIdentifier, TArray<float>& OutTimes, TArray<int32>& Values);
+	static bool Generic_GetAttributeKeys(UAnimDataModel* AnimationDataModel, const FAnimationAttributeIdentifier& AttributeIdentifier, TArray<float>& Times, void* ValuesArray, const FArrayProperty* ValuesArrayProperty);
+	DECLARE_FUNCTION(execGetAttributeKeys);
+};
diff --git a/Engine/Source/Editor/AnimationBlueprintLibrary/Public/K2_AnimationAttributeNodes.h b/Engine/Source/Editor/AnimationBlueprintLibrary/Public/K2_AnimationAttributeNodes.h
new file mode 100644
--- /dev/null
+++ b/Engine/Source/Editor/AnimationBlueprintLibrary/Public/K2_AnimationAttributeNodes.h
@@ -0,0 +1,68 @@
+// Copyright Epic Games, Inc. All Rights Reserved.
+
+#pragma once
+
+#include "CoreMinimal.h"
+#include "K2Node_CallFunction.h"
+#include "K2_AnimationAttributeNodes.generated.h"
+
+class UEdGraph;
+
+/** Base node implementation to reduce duplicated behaviour for different BlueprintLibrary functions */
+UCLASS(abstract)
+class UK2Node_BaseAttributeActionNode : public UK2Node_CallFunction
+{
+	GENERATED_BODY()
+
+public:
+	UK2Node_BaseAttributeActionNode() {}
+	
+	//~ Begin UEdGraphNode Interface.
+	virtual FText GetNodeTitle(ENodeTitleType::Type TitleType) const override;
+	//~ End UEdGraphNode Interface.
+
+	//~ Begin UK2Node Interface
+	virtual void GetMenuActions(FBlueprintActionDatabaseRegistrar& ActionRegistrar) const override;
+	virtual bool IsConnectionDisallowed(const UEdGraphPin* MyPin, const UEdGraphPin* OtherPin, FString& OutReason) const override;
+	//~ End UK2Node Interface
+
+protected:
+	FName AttributeValuePinName = NAME_None;
+	FText AttributeActionFormat;
+	
+	static FName ValuePinName;
+	static FName ValuesPinName;
+};
+
+UCLASS()
+class UK2Node_SetAttributeKeyGeneric : public UK2Node_BaseAttributeActionNode
+{
+	GENERATED_BODY()
+public:
+	UK2Node_SetAttributeKeyGeneric();
+};
+
+UCLASS()
+class UK2Node_SetAttributeKeysGeneric : public UK2Node_BaseAttributeActionNode
+{
+	GENERATED_BODY()
+public:
+	UK2Node_SetAttributeKeysGeneric();
+};
+
+UCLASS()
+class UK2Node_GetAttributeKeyGeneric : public UK2Node_BaseAttributeActionNode
+{
+	GENERATED_BODY()
+public:
+	UK2Node_GetAttributeKeyGeneric();
+};
+
+UCLASS()
+class UK2Node_GetAttributeKeysGeneric : public UK2Node_BaseAttributeActionNode
+{
+	GENERATED_BODY()
+public:
+	UK2Node_GetAttributeKeysGeneric();
+};
+
diff --git a/Engine/Source/Runtime/Engine/Classes/Animation/AttributeCurve.cpp b/Engine/Source/Runtime/Engine/Classes/Animation/AttributeCurve.cpp
--- a/Engine/Source/Runtime/Engine/Classes/Animation/AttributeCurve.cpp
+++ b/Engine/Source/Runtime/Engine/Classes/Animation/AttributeCurve.cpp
@@ -218,7 +218,7 @@ FAttributeKey& FAttributeCurve::GetKey(FKeyHandle KeyHandle)
 	return Keys[GetIndex(KeyHandle)];
 }
 
-FAttributeKey FAttributeCurve::GetKey(FKeyHandle KeyHandle) const
+const FAttributeKey& FAttributeCurve::GetKey(FKeyHandle KeyHandle) const
 {
 	EnsureAllIndicesHaveHandles();
 	return Keys[GetIndex(KeyHandle)];
diff --git a/Engine/Source/Runtime/Engine/Classes/Animation/AttributeCurve.h b/Engine/Source/Runtime/Engine/Classes/Animation/AttributeCurve.h
--- a/Engine/Source/Runtime/Engine/Classes/Animation/AttributeCurve.h
+++ b/Engine/Source/Runtime/Engine/Classes/Animation/AttributeCurve.h
@@ -120,10 +120,17 @@ public:
 		check(AttributeType::StaticStruct() == ScriptStruct);
 		return UpdateOrAddKey(InTime, &InValue, KeyTimeTolerance);
 	}
+
+	/** Finds the key at InTime, and updates its typed value. If it can't find the key within the KeyTimeTolerance, it adds one at that time */
+	FKeyHandle UpdateOrAddTypedKey(float InTime, const void* InValue, const UScriptStruct* ValueType, float KeyTimeTolerance = KINDA_SMALL_NUMBER)
+	{
+		check(ValueType == ScriptStruct);
+		return UpdateOrAddKey(InTime, InValue, KeyTimeTolerance);
+	}
 			
 	/** Functions for getting keys based on handles */
 	FAttributeKey& GetKey(FKeyHandle KeyHandle);
-	FAttributeKey GetKey(FKeyHandle KeyHandle) const;
+	const FAttributeKey& GetKey(FKeyHandle KeyHandle) const;
 
 	/** Finds the key at KeyTime and returns its handle. If it can't find the key within the KeyTimeTolerance, it will return an invalid handle */
 	FKeyHandle FindKey(float KeyTime, float KeyTimeTolerance = KINDA_SMALL_NUMBER) const;
@@ -159,7 +166,7 @@ protected:
 	/** Finds the key at InTime, and updates its typed value. If it can't find the key within the KeyTimeTolerance, it adds one at that time */
 	FKeyHandle UpdateOrAddKey(float InTime, const void* InValue, float KeyTimeTolerance = KINDA_SMALL_NUMBER);
 
-	/** Add a new raw memory key (should be appropriatedly sized) to the curve with the supplied Time and Value. */
+	/** Add a new raw memory key (should be appropriately sized) to the curve with the supplied Time and Value. */
 	FKeyHandle AddKey(float InTime, const void* InValue, FKeyHandle InKeyHandle = FKeyHandle());
 protected:
 	/** The keys, ordered by time */
diff --git a/Engine/Source/Runtime/Engine/Classes/Animation/BuiltInAttributeTypes.h b/Engine/Source/Runtime/Engine/Classes/Animation/BuiltInAttributeTypes.h
--- a/Engine/Source/Runtime/Engine/Classes/Animation/BuiltInAttributeTypes.h
+++ b/Engine/Source/Runtime/Engine/Classes/Animation/BuiltInAttributeTypes.h
@@ -12,12 +12,12 @@
 #include "BuiltInAttributeTypes.generated.h"
 
 /** Attribute type supporting the legacy TVariant<float> atttributes */
-USTRUCT()
+USTRUCT(BlueprintType)
 struct FFloatAnimationAttribute 
 {
 	GENERATED_BODY()
 
-	UPROPERTY()
+	UPROPERTY(EditAnywhere, BlueprintReadWrite, Category=FloatAnimationAttribute)
 	float Value = 0.f;
 
 	void Accumulate(const FFloatAnimationAttribute& Attribute, float Weight, EAdditiveAnimationType AdditiveType)
@@ -45,12 +45,12 @@ struct FFloatAnimationAttribute
 };
 
 /** Attribute type supporting the legacy TVariant<int32> atttributes */
-USTRUCT()
+USTRUCT(BlueprintType)
 struct FIntegerAnimationAttribute
 {
 	GENERATED_BODY()
 
-	UPROPERTY()
+	UPROPERTY(EditAnywhere, BlueprintReadWrite, Category=IntegerAnimationAttribute)
 	int32 Value = 0;
 
 	void Accumulate(const FIntegerAnimationAttribute& Attribute, float Weight, EAdditiveAnimationType AdditiveType)
@@ -78,22 +78,22 @@ struct FIntegerAnimationAttribute
 };
 
 /** Attribute type supporting the legacy TVariant<FString> atttributes */
-USTRUCT()
+USTRUCT(BlueprintType)
 struct FStringAnimationAttribute
 {
 	GENERATED_BODY()
 
-	UPROPERTY()
+	UPROPERTY(EditAnywhere, BlueprintReadWrite, Category=StringAnimationAttribute)
 	FString Value;
 };
 
 /** Attribute type supporting the legacy TVariant<FString> atttributes */
-USTRUCT()
+USTRUCT(BlueprintType)
 struct FTransformAnimationAttribute
 {
 	GENERATED_BODY()
 
-	UPROPERTY()
+	UPROPERTY(EditAnywhere, BlueprintReadWrite, Category=TransformAnimationAttribute)
 	FTransform Value;
 
 	void Accumulate(const FTransformAnimationAttribute& Attribute, float Weight, EAdditiveAnimationType AdditiveType)
