AnimPose API improvements
- Add evaluating / retrieving of curve values from pose
- Add ability to retrieve Additive animation as full-pose (base + additive) or just additive
- Add ability to retrieve set of poses, evaluated at provided equal set of time intervals

diff --git a/Engine/Source/Editor/AnimationBlueprintLibrary/Private/AnimPose.cpp b/Engine/Source/Editor/AnimationBlueprintLibrary/Private/AnimPose.cpp
--- a/Engine/Source/Editor/AnimationBlueprintLibrary/Private/AnimPose.cpp
+++ b/Engine/Source/Editor/AnimationBlueprintLibrary/Private/AnimPose.cpp
@@ -12,26 +12,25 @@ DEFINE_LOG_CATEGORY_STATIC(LogAnimationPoseScripting, Verbose, All);
 void FAnimPose::Init(const FBoneContainer& InBoneContainer)
 {
 	Reset();
-	BoneContainer = InBoneContainer;
 
-	const FReferenceSkeleton& RefSkeleton = BoneContainer.GetSkeletonAsset()->GetReferenceSkeleton();
+	const FReferenceSkeleton& RefSkeleton = InBoneContainer.GetSkeletonAsset()->GetReferenceSkeleton();
 	
-	for (const FBoneIndexType BoneIndex : BoneContainer.GetBoneIndicesArray())
+	for (const FBoneIndexType BoneIndex : InBoneContainer.GetBoneIndicesArray())
 	{			
 		const FCompactPoseBoneIndex CompactIndex(BoneIndex);
-		const FCompactPoseBoneIndex CompactParentIndex = BoneContainer.GetParentBoneIndex(CompactIndex);
+		const FCompactPoseBoneIndex CompactParentIndex = InBoneContainer.GetParentBoneIndex(CompactIndex);
 
-		const int32 SkeletonBoneIndex = BoneContainer.GetSkeletonIndex(CompactIndex);
+		const int32 SkeletonBoneIndex = InBoneContainer.GetSkeletonIndex(CompactIndex);
 		if (SkeletonBoneIndex != INDEX_NONE)
 		{
-			const int32 ParentBoneIndex = CompactParentIndex.GetInt() != INDEX_NONE ? BoneContainer.GetSkeletonIndex(CompactParentIndex) : INDEX_NONE;
+			const int32 ParentBoneIndex = CompactParentIndex.GetInt() != INDEX_NONE ? InBoneContainer.GetSkeletonIndex(CompactParentIndex) : INDEX_NONE;
 
 			BoneIndices.Add(SkeletonBoneIndex);
 			ParentBoneIndices.Add(ParentBoneIndex);
 
 			BoneNames.Add(RefSkeleton.GetBoneName(SkeletonBoneIndex));
 
-			RefLocalSpacePoses.Add(BoneContainer.GetRefPoseTransform(FCompactPoseBoneIndex(BoneIndex)));
+			RefLocalSpacePoses.Add(InBoneContainer.GetRefPoseTransform(FCompactPoseBoneIndex(BoneIndex)));
 		}
 	}
 
@@ -123,8 +122,9 @@ void FAnimPose::GenerateWorldSpaceTransforms()
 	}
 }
 
-void FAnimPose::SetPose(const FCompactPose& CompactPose)
+void FAnimPose::SetPose(const FAnimationPoseData& PoseData)
 {
+	const FCompactPose& CompactPose = PoseData.GetPose();
 	if (IsInitialized())
 	{
 		const FBoneContainer& ContextBoneContainer = CompactPose.GetBoneContainer();
@@ -138,6 +138,25 @@ void FAnimPose::SetPose(const FCompactPose& CompactPose)
 
 		ensure(LocalSpacePoses.Num() == RefLocalSpacePoses.Num());
 		GenerateWorldSpaceTransforms();
+
+		const FBlendedCurve& Curve = PoseData.GetCurve();
+		for (TConstSetBitIterator It(Curve.ValidCurveWeights); It; ++It)
+		{
+			const int32 EntryIndex = It.GetIndex();
+			const uint16 CurveNameUID = (*Curve.UIDToArrayIndexLUT).IsValidIndex(EntryIndex) ? (*Curve.UIDToArrayIndexLUT)[EntryIndex] : INDEX_NONE;
+
+			FSmartName CurveSmartName;
+			const USkeleton* Skeleton = ContextBoneContainer.GetSkeletonAsset();
+			if(Skeleton->GetSmartNameByUID(USkeleton::AnimCurveMappingName, CurveNameUID, CurveSmartName))
+			{
+				CurveNames.Add(CurveSmartName.DisplayName);
+				CurveValues.Add(Curve.CurveWeights[EntryIndex]);
+			}
+			else
+			{
+				ensureMsgf(false, TEXT("Unable to find SmartName for Curve with UID %i from Skeleton %s"), CurveNameUID, *Skeleton->GetPathName());
+			}
+		}
 	}
 	else
 	{
@@ -169,7 +188,6 @@ bool FAnimPose::IsValid() const
 	bIsValid &= WorldSpacePoses.Num() == ExpectedNumBones;
 	bIsValid &= RefLocalSpacePoses.Num() == ExpectedNumBones;
 	bIsValid &= RefWorldSpacePoses.Num() == ExpectedNumBones;
-	bIsValid &= BoneContainer.GetNumBones() == ExpectedNumBones;
 	
 	return bIsValid;
 }
@@ -462,6 +480,23 @@ void UAnimPoseExtensions::GetReferencePose(USkeleton* Skeleton, FAnimPose& OutPo
 	}
 }
 
+void UAnimPoseExtensions::GetCurveNames(const FAnimPose& Pose, TArray<FName>& Curves)
+{
+	Curves.Append(Pose.CurveNames);	
+}
+
+float UAnimPoseExtensions::GetCurveWeight(const FAnimPose& Pose, const FName& CurveName)
+{
+	float CurveValue = 0.f;	
+	const int32 CurveIndex = Pose.CurveNames.IndexOfByKey(CurveName);
+	if (CurveIndex != INDEX_NONE)
+	{
+		CurveValue = Pose.CurveValues[CurveIndex];
+	}
+	
+	return CurveValue;
+}
+
 void UAnimPoseExtensions::GetAnimPoseAtFrame(const UAnimSequenceBase* AnimationSequenceBase, int32 FrameIndex, FAnimPoseEvaluationOptions EvaluationOptions, FAnimPose& Pose)
 {
 	float Time = 0.f;
@@ -470,85 +505,123 @@ void UAnimPoseExtensions::GetAnimPoseAtFrame(const UAnimSequenceBase* AnimationS
 }
 
 void UAnimPoseExtensions::GetAnimPoseAtTime(const UAnimSequenceBase* AnimationSequenceBase, float Time, FAnimPoseEvaluationOptions EvaluationOptions, FAnimPose& Pose)
+{
+	TArray<FAnimPose> InOutPoses;
+	GetAnimPoseAtTimeIntervals(AnimationSequenceBase, { Time }, EvaluationOptions, InOutPoses);
+
+	if (InOutPoses.Num())
+	{
+		ensure(InOutPoses.Num() == 1);
+		Pose = InOutPoses[0];
+	}
+} 
+
+void UAnimPoseExtensions::GetAnimPoseAtTimeIntervals(const UAnimSequenceBase* AnimationSequenceBase, TArray<float> TimeIntervals, FAnimPoseEvaluationOptions EvaluationOptions, TArray<FAnimPose>& InOutPoses)
 {
 	if (AnimationSequenceBase && AnimationSequenceBase->GetSkeleton())
 	{
 		FMemMark Mark(FMemStack::Get());
+		
+		// asset to use for retarget proportions (can be either USkeletalMesh or USkeleton)
+		UObject* AssetToUse;
+		int32 NumRequiredBones;
+		if (EvaluationOptions.OptionalSkeletalMesh)
+		{
+			AssetToUse = CastChecked<UObject>(EvaluationOptions.OptionalSkeletalMesh);
+			NumRequiredBones = EvaluationOptions.OptionalSkeletalMesh->GetRefSkeleton().GetNum();	
+		}
+		else
+		{
+			AssetToUse = CastChecked<UObject>(AnimationSequenceBase->GetSkeleton());
+			NumRequiredBones = AnimationSequenceBase->GetSkeleton()->GetReferenceSkeleton().GetNum();
+		}
 
-		bool bValidTime = false;
-		UAnimationBlueprintLibrary::IsValidTime(AnimationSequenceBase, Time, bValidTime);
-		if (bValidTime)
-		{	
-			// asset to use for retarget proportions (can be either USkeletalMesh or USkeleton)
-			UObject* AssetToUse;
-			int32 NumRequiredBones;
-			if (EvaluationOptions.OptionalSkeletalMesh)
-			{
-				AssetToUse = CastChecked<UObject>(EvaluationOptions.OptionalSkeletalMesh);
-				NumRequiredBones = EvaluationOptions.OptionalSkeletalMesh->GetRefSkeleton().GetNum();	
-			}
-			else
-			{
-				AssetToUse = CastChecked<UObject>(AnimationSequenceBase->GetSkeleton());
-				NumRequiredBones = AnimationSequenceBase->GetSkeleton()->GetReferenceSkeleton().GetNum();
-			}
-
-			TArray<FBoneIndexType> RequiredBoneIndexArray;
-			RequiredBoneIndexArray.AddUninitialized(NumRequiredBones);
-			for (int32 BoneIndex = 0; BoneIndex < RequiredBoneIndexArray.Num(); ++BoneIndex)
-			{
-				RequiredBoneIndexArray[BoneIndex] = BoneIndex;
-			}
+		TArray<FBoneIndexType> RequiredBoneIndexArray;
+		RequiredBoneIndexArray.AddUninitialized(NumRequiredBones);
+		for (int32 BoneIndex = 0; BoneIndex < RequiredBoneIndexArray.Num(); ++BoneIndex)
+		{
+			RequiredBoneIndexArray[BoneIndex] = BoneIndex;
+		}
 
-			FBoneContainer RequiredBones;
-			RequiredBones.InitializeTo(RequiredBoneIndexArray, FCurveEvaluationOption(false), *AssetToUse);
-			
-			RequiredBones.SetUseRAWData(EvaluationOptions.EvaluationType == EAnimDataEvalType::Raw);
-			RequiredBones.SetUseSourceData(EvaluationOptions.EvaluationType == EAnimDataEvalType::Source);
-			
-			RequiredBones.SetDisableRetargeting(!EvaluationOptions.bShouldRetarget);
+		FBoneContainer RequiredBones;
+		RequiredBones.InitializeTo(RequiredBoneIndexArray, FCurveEvaluationOption(EvaluationOptions.bEvaluateCurves), *AssetToUse);
+		
+		RequiredBones.SetUseRAWData(EvaluationOptions.EvaluationType == EAnimDataEvalType::Raw);
+		RequiredBones.SetUseSourceData(EvaluationOptions.EvaluationType == EAnimDataEvalType::Source);
+		
+		RequiredBones.SetDisableRetargeting(!EvaluationOptions.bShouldRetarget);
+		
+		FCompactPose CompactPose;
+        FBlendedCurve Curve;
+        UE::Anim::FStackAttributeContainer Attributes;
+
+        FAnimationPoseData PoseData(CompactPose, Curve, Attributes);
+        FAnimExtractContext Context(0.f, EvaluationOptions.bExtractRootMotion);
+    
+        FCompactPose BasePose;
+        BasePose.SetBoneContainer(&RequiredBones);
+        
+        CompactPose.SetBoneContainer(&RequiredBones);
+        Curve.InitFrom(RequiredBones);
+
+		FAnimPose Pose;
+		Pose.Init(RequiredBones);
+		
+		for (int32 Index = 0; Index < TimeIntervals.Num(); ++Index)
+		{
+			const float EvalInterval = TimeIntervals[Index];
 			
-			FCompactPose CompactPose;
-			CompactPose.SetBoneContainer(&RequiredBones);
+			bool bValidTime = false;
+			UAnimationBlueprintLibrary::IsValidTime(AnimationSequenceBase, EvalInterval, bValidTime);
+			ensure(bValidTime);
 
-			Pose.Init(CompactPose.GetBoneContainer());
+			Context.CurrentTime = EvalInterval;
 
-			FBlendedCurve Curve;
+			FAnimPose& FramePose = InOutPoses.AddDefaulted_GetRef();
+			FramePose = Pose;
+			
 			Curve.InitFrom(RequiredBones);
-			UE::Anim::FStackAttributeContainer Attributes;
-
-			FAnimationPoseData PoseData(CompactPose, Curve, Attributes);
-			FAnimExtractContext Context(Time, EvaluationOptions.bExtractRootMotion);
 
-			AnimationSequenceBase->GetAnimationPose(PoseData, Context);
-
-			if (AnimationSequenceBase->IsValidAdditive())
+			if (bValidTime)
 			{
-				const UAnimSequence* AnimSequence = Cast<const UAnimSequence>(AnimationSequenceBase);
-				FCompactPose BasePose;
-				BasePose.SetBoneContainer(&RequiredBones);
+				if (AnimationSequenceBase->IsValidAdditive())
+				{
+					CompactPose.ResetToAdditiveIdentity();
+					AnimationSequenceBase->GetAnimationPose(PoseData, Context);
 
-				FBlendedCurve BaseCurve;
-				BaseCurve.InitFrom(RequiredBones);
-				UE::Anim::FStackAttributeContainer BaseAttributes;
+					if (EvaluationOptions.bRetrieveAdditiveAsFullPose)
+					{
+						const UAnimSequence* AnimSequence = Cast<const UAnimSequence>(AnimationSequenceBase);
+					
+						FBlendedCurve BaseCurve;
+						BaseCurve.InitFrom(RequiredBones);
+						UE::Anim::FStackAttributeContainer BaseAttributes;
 				
-				FAnimationPoseData BasePoseData(BasePose, BaseCurve, BaseAttributes);
-				AnimSequence->GetAdditiveBasePose(BasePoseData, Context);
+						FAnimationPoseData BasePoseData(BasePose, BaseCurve, BaseAttributes);
+						AnimSequence->GetAdditiveBasePose(BasePoseData, Context);
 
-				FAnimationRuntime::AccumulateAdditivePose(BasePoseData, PoseData, 1.f, AnimSequence->GetAdditiveAnimType());
-				BasePose.NormalizeRotations();
+						FAnimationRuntime::AccumulateAdditivePose(BasePoseData, PoseData, 1.f, AnimSequence->GetAdditiveAnimType());
+						BasePose.NormalizeRotations();
 				
-				Pose.SetPose(BasePose);
+						FramePose.SetPose(BasePoseData);
+					}
+					else
+					{
+						FramePose.SetPose(PoseData);
+					}
+				}
+				else
+				{
+					CompactPose.ResetToRefPose();
+					AnimationSequenceBase->GetAnimationPose(PoseData, Context);
+					FramePose.SetPose(PoseData);
+				}
 			}
 			else
 			{
-				Pose.SetPose(CompactPose);
+				UE_LOG(LogAnimationPoseScripting, Warning, TEXT("Invalid time value %f for Animation Sequence %s supplied for GetBonePosesForTime"), EvalInterval, *AnimationSequenceBase->GetName());
 			}
 		}
-		else
-		{
-			UE_LOG(LogAnimationPoseScripting, Warning, TEXT("Invalid time value %f for Animation Sequence %s supplied for GetBonePosesForTime"), Time, *AnimationSequenceBase->GetName());
-		}
 	}
 	else
 	{
diff --git a/Engine/Source/Editor/AnimationBlueprintLibrary/Public/AnimPose.h b/Engine/Source/Editor/AnimationBlueprintLibrary/Public/AnimPose.h
--- a/Engine/Source/Editor/AnimationBlueprintLibrary/Public/AnimPose.h
+++ b/Engine/Source/Editor/AnimationBlueprintLibrary/Public/AnimPose.h
@@ -8,11 +8,6 @@
 
 #include "AnimPose.generated.h"
 
-/*	TODO
- *	- Automated tests
- *	- Verifying retargeted animation between different skeletons
- */
-
 UENUM(BlueprintType)
 enum class EAnimPoseSpaces : uint8
 {
@@ -53,6 +48,14 @@ struct ANIMATIONBLUEPRINTLIBRARY_API FAnimPoseEvaluationOptions
 	// Optional skeletal mesh with proportions to use when evaluating a pose
 	UPROPERTY(EditAnywhere, AdvancedDisplay, BlueprintReadWrite, Category="Animation|Pose")
 	TObjectPtr<USkeletalMesh> OptionalSkeletalMesh = nullptr;
+	
+	// Whether or additive animations should be applied to their base-pose 
+	UPROPERTY(EditAnywhere, AdvancedDisplay, BlueprintReadWrite, Category="Animation|Pose")
+	bool bRetrieveAdditiveAsFullPose = true;
+
+	// Whether or not to evaluate Animation Curves
+	UPROPERTY(EditAnywhere, AdvancedDisplay, BlueprintReadWrite, Category="Animation|Pose")
+	bool bEvaluateCurves = true;
 };
 
 /** Script friendly representation of an evaluated animation bone pose */
@@ -72,7 +75,7 @@ protected:
 	/** Generates the contained bone data using the provided Component and its AnimInstance */
 	void SetPose(USkeletalMeshComponent* Component);
 	/** Generates the contained bone data using the provided CompactPose */
-	void SetPose(const FCompactPose& CompactPose);
+	void SetPose(const FAnimationPoseData& PoseData);
 	/** Copies the reference pose to animated pose data */
 	void SetToRefPose();
 
@@ -108,7 +111,11 @@ protected:
 	UPROPERTY()
 	TArray<FTransform> RefWorldSpacePoses;
 
-	FBoneContainer BoneContainer;
+	UPROPERTY()
+	TArray<FName> CurveNames;
+
+	UPROPERTY()
+	TArray<float> CurveValues;
 
 	friend class UAnimPoseExtensions;
 };
@@ -130,7 +137,6 @@ public:
 	UFUNCTION(BlueprintPure, meta = (ScriptMethod), Category = "Animation|Pose")
 	static bool IsValid(UPARAM(ref) const FAnimPose& Pose);
 
-	/** */
 	/**
 	* Returns an array of bone names contained by the pose
 	*
@@ -225,6 +231,16 @@ public:
 	UFUNCTION(BlueprintCallable, meta = (ScriptMethod), Category = "Animation|Pose")
 	static void GetAnimPoseAtTime(const UAnimSequenceBase* AnimationSequenceBase, float Time, FAnimPoseEvaluationOptions EvaluationOptions, FAnimPose& Pose);
 
+	/**
+	* Evaluates an Animation Sequence Base at different time intervals to generate a valid Anim Pose instances
+	*
+	* @param	AnimationSequenceBase	Animation sequence base to evaluate the pose from
+	* @param	TimeIntervals			Times at which the pose should be evaluated
+	* @param	EvaluationOptions		Options determining the way the pose should be evaluated
+	* @param	InOutPoses				Anim Poses holding the evaluated data (number matches TimeIntervals)
+	*/
+	static void GetAnimPoseAtTimeIntervals(const UAnimSequenceBase* AnimationSequenceBase, TArray<float> TimeIntervals, FAnimPoseEvaluationOptions EvaluationOptions, TArray<FAnimPose>& InOutPoses);
+
 	/**
 	* Evaluates an Animation Sequence Base to generate a valid Anim Pose instance
 	*
@@ -255,4 +271,25 @@ public:
 	*/
 	UFUNCTION(BlueprintPure, meta = (ScriptMethod), Category = "Animation|Pose")
 	static void GetReferencePose(USkeleton* Skeleton, FAnimPose& OutPose);
+
+	/**
+	* Returns an array of curve names contained by the pose
+	*
+	* @param	Pose	Anim Pose to retrieve the names from
+	* @param	Curves	Array to be populated with the curve names
+	*
+	*/
+	UFUNCTION(BlueprintPure, meta = (ScriptMethod), Category = "Animation|Pose")
+	static void GetCurveNames(UPARAM(ref) const FAnimPose& Pose, TArray<FName>& Curves);
+	
+	/**
+    * Returns the weight of an evaluated curve - if found
+    *
+    * @param	Pose		Anim Pose to retrieve the value from
+    * @param	CurveName	Curve to retrieve the weight value for
+    * 
+    * @return	Curve weight value, if found - 0.f otherwise
+    */
+    UFUNCTION(BlueprintPure, meta = (ScriptMethod), Category = "Animation|Pose")
+    static float GetCurveWeight(UPARAM(ref) const FAnimPose& Pose, const FName& CurveName);
 };
\ No newline at end of file
