Add GetMontageSlotNames to AnimationBlueprintLibrary to retrieve the name of all slots used by a given montage

diff --git a/Engine/Source/Editor/AnimationBlueprintLibrary/Private/AnimationBlueprintLibrary.cpp b/Engine/Source/Editor/AnimationBlueprintLibrary/Private/AnimationBlueprintLibrary.cpp
--- a/Engine/Source/Editor/AnimationBlueprintLibrary/Private/AnimationBlueprintLibrary.cpp
+++ b/Engine/Source/Editor/AnimationBlueprintLibrary/Private/AnimationBlueprintLibrary.cpp
@@ -75,6 +75,23 @@ void UAnimationBlueprintLibrary::GetAnimationTrackNames(const UAnimSequenceBase*
 	}	
 }
 
+void UAnimationBlueprintLibrary::GetMontageSlotNames(const UAnimMontage* AnimationMontage, TArray<FName>& SlotNames)
+{
+	SlotNames.Empty();
+	if (AnimationMontage)
+	{
+		for (int32 SlotIdx = 0; SlotIdx < AnimationMontage->SlotAnimTracks.Num(); SlotIdx++)
+		{
+			const FSlotAnimationTrack& SlotAnimTrack = AnimationMontage->SlotAnimTracks[SlotIdx];
+			SlotNames.Add(SlotAnimTrack.SlotName);
+		}
+	}
+	else
+	{
+		UE_LOG(LogAnimationBlueprintLibrary, Warning, TEXT("Invalid Animation Montage supplied for GetMontageSlotNames"));
+	}
+}
+
 void UAnimationBlueprintLibrary::GetAnimationCurveNames(const UAnimSequence* AnimationSequence, ERawCurveTrackTypes CurveType, TArray<FName>& CurveNames)
 {
 	CurveNames.Empty();
diff --git a/Engine/Source/Editor/AnimationBlueprintLibrary/Public/AnimationBlueprintLibrary.h b/Engine/Source/Editor/AnimationBlueprintLibrary/Public/AnimationBlueprintLibrary.h
--- a/Engine/Source/Editor/AnimationBlueprintLibrary/Public/AnimationBlueprintLibrary.h
+++ b/Engine/Source/Editor/AnimationBlueprintLibrary/Public/AnimationBlueprintLibrary.h
@@ -50,6 +50,10 @@ public:
 	UFUNCTION(BlueprintPure, Category = "AnimationBlueprintLibrary|Animation")
 	static void GetAnimationTrackNames(const UAnimSequenceBase* AnimationSequenceBase, TArray<FName>& TrackNames);
 
+	/** Retrieves the Names of the Animation Slots used in the given Montage */
+	UFUNCTION(BlueprintPure, Category = "AnimationBlueprintLibrary|Montage")
+	static void GetMontageSlotNames(const UAnimMontage* AnimationMontage, TArray<FName>& SlotNames);
+
 	/** Retrieves the Names of the individual float curves for the given Animation Sequence */
 	UFUNCTION(BlueprintPure, Category = "AnimationBlueprintLibrary|Animation")
 	static void GetAnimationCurveNames(const UAnimSequence* AnimationSequence, ERawCurveTrackTypes CurveType, TArray<FName>& CurveNames);
