UE5/Release-5.1 - Duplicated keys for localization - UE - Anim

diff --git a/Engine/Source/Editor/AnimGraph/Private/AnimGraphNode_LinkedAnimLayer.cpp b/Engine/Source/Editor/AnimGraph/Private/AnimGraphNode_LinkedAnimLayer.cpp
--- a/Engine/Source/Editor/AnimGraph/Private/AnimGraphNode_LinkedAnimLayer.cpp
+++ b/Engine/Source/Editor/AnimGraph/Private/AnimGraphNode_LinkedAnimLayer.cpp
@@ -102,16 +102,17 @@ FText UAnimGraphNode_LinkedAnimLayer::GetNodeTitle(ENodeTitleType::Type TitleTyp
 	UClass* TargetClass = *Node.Interface;
 	UAnimBlueprint* TargetAnimBlueprintInterface = TargetClass ? CastChecked<UAnimBlueprint>(TargetClass->ClassGeneratedBy) : nullptr;
 
+	const FText DefaultNodeTitle = LOCTEXT("NodeTitle", "Linked Anim Layer");
 	if(TitleType == ENodeTitleType::MenuTitle)
 	{
-		return LOCTEXT("NodeTitle", "Linked Anim Layer");
+		return DefaultNodeTitle;
 	}
 	else
 	{
 		bool bIsSelf = TargetAnimBlueprintInterface == nullptr; 
 		
 		FFormatNamedArguments Args;
-		Args.Add(TEXT("NodeType"), bIsSelf ? LOCTEXT("NodeTitle", "Anim Layer (self)") : LOCTEXT("NodeTitle", "Linked Anim Layer"));
+		Args.Add(TEXT("NodeType"), bIsSelf ? LOCTEXT("SelfLayerNodeTitle", "Anim Layer (self)") : DefaultNodeTitle);
 		Args.Add(TEXT("TargetClass"), bIsSelf ? LOCTEXT("ClassSelf", "Self") : FText::FromString(TargetAnimBlueprintInterface->GetName()));
 		Args.Add(TEXT("Layer"), (Node.Layer == NAME_None) ? LOCTEXT("LayerNone", "None") : FText::FromName(Node.Layer));
 
diff --git a/Engine/Source/Editor/AnimationBlueprintLibrary/Private/AnimationAttributeBlueprintLibrary.cpp b/Engine/Source/Editor/AnimationBlueprintLibrary/Private/AnimationAttributeBlueprintLibrary.cpp
--- a/Engine/Source/Editor/AnimationBlueprintLibrary/Private/AnimationAttributeBlueprintLibrary.cpp
+++ b/Engine/Source/Editor/AnimationBlueprintLibrary/Private/AnimationAttributeBlueprintLibrary.cpp
@@ -113,7 +113,7 @@ DEFINE_FUNCTION(UAnimationAttributeBlueprintLibrary::execSetAttributeKeys)
 	{
 		FBlueprintExceptionInfo ExceptionInfo(
 			EBlueprintExceptionType::AccessViolation, 
-			LOCTEXT("SetAttributeKey_InvalidController", "Accessed None attempting to call SetAttributeKeys.")
+			LOCTEXT("SetAttributeKeys_InvalidController", "Accessed None attempting to call SetAttributeKeys.")
 		);
 		FBlueprintCoreDelegates::ThrowScriptException(P_THIS, Stack, ExceptionInfo);
 	}
@@ -181,7 +181,7 @@ DEFINE_FUNCTION(UAnimationAttributeBlueprintLibrary::execGetAttributeKey)
 	{
 		FBlueprintExceptionInfo ExceptionInfo(
 			EBlueprintExceptionType::AccessViolation,
-			LOCTEXT("SetAttributeKey_InvalidValue", "Failed to resolve the attribute value parameter for GetAttributeKey.")
+			LOCTEXT("GetAttributeKey_InvalidValue", "Failed to resolve the attribute value parameter for GetAttributeKey.")
 		);
 		FBlueprintCoreDelegates::ThrowScriptException(P_THIS, Stack, ExceptionInfo);
 	}
@@ -190,7 +190,7 @@ DEFINE_FUNCTION(UAnimationAttributeBlueprintLibrary::execGetAttributeKey)
 	{
 		FBlueprintExceptionInfo ExceptionInfo(
 			EBlueprintExceptionType::AccessViolation, 
-			LOCTEXT("SetAttributeKey_InvalidController", "Accessed None attempting to call GetAttributeKey.")
+			LOCTEXT("GetAttributeKey_InvalidController", "Accessed None attempting to call GetAttributeKey.")
 		);
 		FBlueprintCoreDelegates::ThrowScriptException(P_THIS, Stack, ExceptionInfo);
 	}
@@ -266,7 +266,7 @@ DEFINE_FUNCTION(UAnimationAttributeBlueprintLibrary::execGetAttributeKeys)
 	{
 		FBlueprintExceptionInfo ExceptionInfo(
 			EBlueprintExceptionType::AccessViolation,
-			LOCTEXT("SetAttributeKeys_InvalidValue", "Failed to resolve the attribute values parameter for GetAttributeKeys.")
+			LOCTEXT("GetAttributeKeys_InvalidValue", "Failed to resolve the attribute values parameter for GetAttributeKeys.")
 		);
 		FBlueprintCoreDelegates::ThrowScriptException(P_THIS, Stack, ExceptionInfo);
 	}
@@ -275,7 +275,7 @@ DEFINE_FUNCTION(UAnimationAttributeBlueprintLibrary::execGetAttributeKeys)
 	{
 		FBlueprintExceptionInfo ExceptionInfo(
 			EBlueprintExceptionType::AccessViolation, 
-			LOCTEXT("SetAttributeKey_InvalidController", "Accessed None attempting to call GetAttributeKeys.")
+			LOCTEXT("GetAttributeKeys_InvalidController", "Accessed None attempting to call GetAttributeKeys.")
 		);
 		FBlueprintCoreDelegates::ThrowScriptException(P_THIS, Stack, ExceptionInfo);
 	}
