ANIM-RT truncation fixes - AnimationBlueprintLibrary

diff --git a/Engine/Source/Editor/AnimationBlueprintLibrary/Private/AnimPose.cpp b/Engine/Source/Editor/AnimationBlueprintLibrary/Private/AnimPose.cpp
--- a/Engine/Source/Editor/AnimationBlueprintLibrary/Private/AnimPose.cpp
+++ b/Engine/Source/Editor/AnimationBlueprintLibrary/Private/AnimPose.cpp
@@ -493,7 +493,7 @@ void UAnimPoseExtensions::GetReferencePose(USkeleton* Skeleton, FAnimPose& OutPo
 		RequiredBoneIndexArray.AddUninitialized(RefSkeleton.GetNum());
 		for (int32 BoneIndex = 0; BoneIndex < RequiredBoneIndexArray.Num(); ++BoneIndex)
 		{
-			RequiredBoneIndexArray[BoneIndex] = BoneIndex;
+			RequiredBoneIndexArray[BoneIndex] = IntCastChecked<FBoneIndexType>(BoneIndex);
 		}
 
 		FBoneContainer RequiredBones;
@@ -567,7 +567,7 @@ void UAnimPoseExtensions::GetAnimPoseAtTimeIntervals(const UAnimSequenceBase* An
 		RequiredBoneIndexArray.AddUninitialized(NumRequiredBones);
 		for (int32 BoneIndex = 0; BoneIndex < RequiredBoneIndexArray.Num(); ++BoneIndex)
 		{
-			RequiredBoneIndexArray[BoneIndex] = BoneIndex;
+			RequiredBoneIndexArray[BoneIndex] = static_cast<FBoneIndexType>(BoneIndex);
 		}
 
 		FBoneContainer RequiredBones;
@@ -599,7 +599,7 @@ void UAnimPoseExtensions::GetAnimPoseAtTimeIntervals(const UAnimSequenceBase* An
 			const double EvalInterval = TimeIntervals[Index];
 			
 			bool bValidTime = false;
-			UAnimationBlueprintLibrary::IsValidTime(AnimationSequenceBase, EvalInterval, bValidTime);
+			UAnimationBlueprintLibrary::IsValidTime(AnimationSequenceBase, static_cast<float>(EvalInterval), bValidTime);
 			ensure(bValidTime);
 
 			Context.CurrentTime = EvalInterval;
diff --git a/Engine/Source/Editor/AnimationBlueprintLibrary/Private/AnimationBlueprintLibrary.cpp b/Engine/Source/Editor/AnimationBlueprintLibrary/Private/AnimationBlueprintLibrary.cpp
--- a/Engine/Source/Editor/AnimationBlueprintLibrary/Private/AnimationBlueprintLibrary.cpp
+++ b/Engine/Source/Editor/AnimationBlueprintLibrary/Private/AnimationBlueprintLibrary.cpp
@@ -2383,7 +2383,7 @@ bool UAnimationBlueprintLibrary::EvaluateRootBoneTimecodeAttributesAtTime(const
 		}
 		else
 		{
-			FrameRate = FFrameRate(FMath::TruncToInt(TimecodeRateAsDecimal), 1);
+			FrameRate = FFrameRate(static_cast<int32>(FMath::TruncToInt(TimecodeRateAsDecimal)), 1);
 		}
 	}
 	else if (const UAnimSequence* AnimSequence = AnimDataModel->GetAnimationSequence())
