Animation Compression

diff --git a/Engine/Plugins/Editor/SequencerAnimTools/Source/SequencerAnimTools/Private/Sequencer/AnimationBoneTrail.cpp b/Engine/Plugins/Editor/SequencerAnimTools/Source/SequencerAnimTools/Private/Sequencer/AnimationBoneTrail.cpp
--- a/Engine/Plugins/Editor/SequencerAnimTools/Source/SequencerAnimTools/Private/Sequencer/AnimationBoneTrail.cpp
+++ b/Engine/Plugins/Editor/SequencerAnimTools/Source/SequencerAnimTools/Private/Sequencer/AnimationBoneTrail.cpp
@@ -97,7 +97,6 @@ void FAnimTrajectoryCache::GetSpaceBasedAnimationData(TArray<TArray<FTransform>>
 
 	// 2d array of animated time [boneindex][time key]
 	const int32 NumKeys = CachedAnimSequence->GetNumberOfSampledKeys();
-	float Interval = UE::SequencerAnimTools::GetIntervalPerKey(NumKeys, CachedAnimSequence->GetPlayLength());
 
 	// allocate arrays
 	for (int32 BoneIndex = 0; BoneIndex < NumBones; ++BoneIndex)
@@ -106,31 +105,28 @@ void FAnimTrajectoryCache::GetSpaceBasedAnimationData(TArray<TArray<FTransform>>
 	}
 
 	// If skeleton to track map isn't initialized yet, initialize it
-	if (SkelToTrackIdx.Num() == 0)
+	if (SkelToBoneName.Num() == 0)
 	{
-		SkelToTrackIdx.SetNumUninitialized(NumBones);
+		SkelToBoneName.SetNumUninitialized(NumBones);
 		for (int32 BoneIndex = 0; BoneIndex < NumBones; ++BoneIndex)
 		{
-			FName BoneName = MySkeleton->GetReferenceSkeleton().GetBoneName(BoneIndex);
-			int32 TrackIndex = CachedAnimSequence->GetDataModel()->GetBoneTrackIndexByName(BoneName);
-			SkelToTrackIdx[BoneIndex] = TrackIndex;
+			const FName SkeletonBoneName = MySkeleton->GetReferenceSkeleton().GetBoneName(BoneIndex);
+			SkelToBoneName[BoneIndex] = CachedAnimSequence->GetDataModel()->IsValidBoneTrackName(SkeletonBoneName) ? SkeletonBoneName : NAME_None;
 		}
 	}
 
 	// Set component transform for each bone
 	for (int32 BoneIndex = 0; BoneIndex < NumBones; ++BoneIndex)
 	{
-		int32 TrackIndex = SkelToTrackIdx[BoneIndex];
-		int32 ParentBoneIndex = MySkeleton->GetReferenceSkeleton().GetParentIndex(BoneIndex);
+		const FName& BoneName = SkelToBoneName[BoneIndex];
+		const int32 ParentBoneIndex = MySkeleton->GetReferenceSkeleton().GetParentIndex(BoneIndex);
 
-		if (TrackIndex != INDEX_NONE)
+		if (BoneName != NAME_None)
 		{
 			// fill up keys - calculate PK1 * K1
 			for (int32 Key = 0; Key < NumKeys; ++Key)
 			{
-				FTransform AnimatedLocalKey;
-				UE::Anim::GetBoneTransformFromModel(CachedAnimSequence->GetDataModel(), AnimatedLocalKey, TrackIndex, Interval * Key, CachedAnimSequence->Interpolation);
-
+				const FTransform AnimatedLocalKey = CachedAnimSequence->GetDataModel()->EvaluateBoneTrackTransform(BoneName, Key, CachedAnimSequence->Interpolation);
 				if (ParentBoneIndex != INDEX_NONE)
 				{
 					OutAnimationDataInComponentSpace[BoneIndex][Key] = AnimatedLocalKey * OutAnimationDataInComponentSpace[ParentBoneIndex][Key];
@@ -144,8 +140,7 @@ void FAnimTrajectoryCache::GetSpaceBasedAnimationData(TArray<TArray<FTransform>>
 		else
 		{
 			// get local spaces from refpose and use that to fill it up
-			FTransform LocalTransform = MySkeleton->GetReferenceSkeleton().GetRefBonePose()[BoneIndex];
-
+			const FTransform LocalTransform = MySkeleton->GetReferenceSkeleton().GetRefBonePose()[BoneIndex];
 			for (int32 Key = 0; Key < NumKeys; ++Key)
 			{
 				if (ParentBoneIndex != INDEX_NONE)
diff --git a/Engine/Plugins/Editor/SequencerAnimTools/Source/SequencerAnimTools/Private/Sequencer/AnimationBoneTrail.h b/Engine/Plugins/Editor/SequencerAnimTools/Source/SequencerAnimTools/Private/Sequencer/AnimationBoneTrail.h
--- a/Engine/Plugins/Editor/SequencerAnimTools/Source/SequencerAnimTools/Private/Sequencer/AnimationBoneTrail.h
+++ b/Engine/Plugins/Editor/SequencerAnimTools/Source/SequencerAnimTools/Private/Sequencer/AnimationBoneTrail.h
@@ -25,7 +25,7 @@ public:
 		, CachedAnimSequence(NewObject<UAnimSequence>())
 		, GlobalBoneTransforms()
 		, ComponentBoneTransforms()
-		, SkelToTrackIdx()
+		, SkelToBoneName()
 		, AnimRange()
 		, Spacing()
 		, bDirty(true)
@@ -56,7 +56,7 @@ private:
 	UAnimSequence* CachedAnimSequence;
 	TArray<TArray<FTransform>> GlobalBoneTransforms;
 	TArray<TArray<FTransform>> ComponentBoneTransforms;
-	TArray<int32> SkelToTrackIdx;
+	TArray<FName> SkelToBoneName;
 	TRange<double> AnimRange;
 	double Spacing;
 	bool bDirty;
diff --git a/Engine/Plugins/Importers/AlembicImporter/Source/AlembicLibrary/Private/AbcImporter.cpp b/Engine/Plugins/Importers/AlembicImporter/Source/AlembicLibrary/Private/AbcImporter.cpp
--- a/Engine/Plugins/Importers/AlembicImporter/Source/AlembicLibrary/Private/AbcImporter.cpp
+++ b/Engine/Plugins/Importers/AlembicImporter/Source/AlembicLibrary/Private/AbcImporter.cpp
@@ -764,7 +764,7 @@ TArray<UObject*> FAbcImporter::ImportAsSkeletalMesh(UObject* InParent, EObjectFl
 				const FReferenceSkeleton& RefSkeleton = SkeletalMesh->GetRefSkeleton();
 				const TArray<FMeshBoneInfo>& BonesInfo = RefSkeleton.GetRawRefBoneInfo();
 				
-				Controller.AddBoneTrack(BonesInfo[0].Name);
+				Controller.AddBoneCurve(BonesInfo[0].Name);
 				Controller.SetBoneTrackKeys(BonesInfo[0].Name, RootBoneTrack.PosKeys, RootBoneTrack.RotKeys, RootBoneTrack.ScaleKeys);
 			}
 
diff --git a/Engine/Plugins/Importers/USDImporter/Source/USDUtilities/Private/USDSkeletalDataConversion.cpp b/Engine/Plugins/Importers/USDImporter/Source/USDUtilities/Private/USDSkeletalDataConversion.cpp
--- a/Engine/Plugins/Importers/USDImporter/Source/USDUtilities/Private/USDSkeletalDataConversion.cpp
+++ b/Engine/Plugins/Importers/USDImporter/Source/USDUtilities/Private/USDSkeletalDataConversion.cpp
@@ -2184,7 +2184,7 @@ bool UsdToUnreal::ConvertSkelAnim(
 
 		for ( int32 BoneIndex = 0; BoneIndex < NumBones; ++BoneIndex )
 		{
-			Controller.AddBoneTrack( BoneInfo[ BoneIndex ].Name, bShouldTransact );
+			Controller.AddBoneCurve( BoneInfo[ BoneIndex ].Name, bShouldTransact );
 			Controller.SetBoneTrackKeys( BoneInfo[ BoneIndex ].Name, JointTracks[ BoneIndex ].PosKeys, JointTracks[ BoneIndex ].RotKeys, JointTracks[ BoneIndex ].ScaleKeys, bShouldTransact );
 		}
 	}
diff --git a/Engine/Plugins/Interchange/Runtime/Source/Import/Private/Animation/InterchangeAnimSequenceFactory.cpp b/Engine/Plugins/Interchange/Runtime/Source/Import/Private/Animation/InterchangeAnimSequenceFactory.cpp
--- a/Engine/Plugins/Interchange/Runtime/Source/Import/Private/Animation/InterchangeAnimSequenceFactory.cpp
+++ b/Engine/Plugins/Interchange/Runtime/Source/Import/Private/Animation/InterchangeAnimSequenceFactory.cpp
@@ -432,7 +432,7 @@ namespace UE::Interchange::Private
 					&& TimeKeys.Num() == BakeKeyCount);
 
 				//add new track
-				Controller.AddBoneTrack(BoneName);
+				Controller.AddBoneCurve(BoneName);
 				Controller.SetBoneTrackKeys(BoneName, RawTrack.PosKeys, RawTrack.RotKeys, RawTrack.ScaleKeys);
 			}
 		}
diff --git a/Engine/Plugins/VirtualProduction/Takes/Source/TakeTrackRecorders/Private/MovieScene3DTransformTrackRecorder.cpp b/Engine/Plugins/VirtualProduction/Takes/Source/TakeTrackRecorders/Private/MovieScene3DTransformTrackRecorder.cpp
--- a/Engine/Plugins/VirtualProduction/Takes/Source/TakeTrackRecorders/Private/MovieScene3DTransformTrackRecorder.cpp
+++ b/Engine/Plugins/VirtualProduction/Takes/Source/TakeTrackRecorders/Private/MovieScene3DTransformTrackRecorder.cpp
@@ -503,11 +503,14 @@ void UMovieScene3DTransformTrackRecorder::PostProcessAnimationData(UMovieSceneAn
 				int32 RootIndex = INDEX_NONE;
 				USkeleton* AnimSkeleton = AnimSequence->GetSkeleton();
 
-				const TArray<FBoneAnimationTrack>& BoneAnimationTracks = AnimSequence->GetDataModel()->GetBoneAnimationTracks();
-				for (const FBoneAnimationTrack& AnimationTrack : BoneAnimationTracks)
+				TArray<FName> TrackNames;
+				AnimSequence->GetDataModelInterface()->GetBoneTrackNames(TrackNames);
+
+				FName RootBoneName = NAME_None;
+				for (const FName& TrackName : TrackNames)
 				{
 					// Verify if this bone exists in skeleton
-					const int32 BoneTreeIndex = AnimationTrack.BoneTreeIndex;
+					const int32 BoneTreeIndex = AnimSequence->GetSkeleton()->GetReferenceSkeleton().FindBoneIndex(TrackName);
 					if (BoneTreeIndex != INDEX_NONE)
 					{
 						const int32 BoneIndex = AnimSkeleton->GetMeshBoneIndexFromSkeletonBoneIndex(SkinnedAsset, BoneTreeIndex);
@@ -516,57 +519,26 @@ void UMovieScene3DTransformTrackRecorder::PostProcessAnimationData(UMovieSceneAn
 						{
 							// We've found the root (root bones do not have a valid parent)
 							RootIndex = BoneIndex;
+							RootBoneName = TrackName;
 							break;
 						}
 					}
 				}
 
-				if (RootIndex == INDEX_NONE)
+				if (RootBoneName == NAME_None)
 				{
 					FString ObjectToRecordName = ObjectToRecord.IsValid() ? ObjectToRecord->GetName() : TEXT("Unnamed_Actor");
 					UE_LOG(LogTakesCore, Log, TEXT("No Root Found for (%s)"), *ObjectToRecordName);
 					return;
 				}
 
-				ensure(BoneAnimationTracks.IsValidIndex(RootIndex));
-
 				const FFrameRate TickResolution = MovieSceneSection->GetTypedOuter<UMovieScene>()->GetTickResolution();
 				const FFrameNumber StartTime = MovieSceneSection->GetInclusiveStartFrame();
 
 				FTransform InvComponentTransform = AnimTrackRecorder->GetComponentTransform().Inverse();
-
-				const FRawAnimSequenceTrack& RawTrack = BoneAnimationTracks[RootIndex].InternalTrackData;
-				const int32 KeyCount = FMath::Max(FMath::Max(RawTrack.PosKeys.Num(), RawTrack.RotKeys.Num()), RawTrack.ScaleKeys.Num());
-				for (int32 KeyIndex = 0; KeyIndex < KeyCount; KeyIndex++)
+				for (int32 KeyIndex = 0; KeyIndex < AnimSequence->GetDataModelInterface()->GetNumberOfKeys(); KeyIndex++)
 				{
-					FTransform Transform;
-					if (RawTrack.PosKeys.IsValidIndex(KeyIndex))
-					{
-						Transform.SetTranslation(FVector(RawTrack.PosKeys[KeyIndex]));
-					}
-					else if (RawTrack.PosKeys.Num() > 0)
-					{
-						Transform.SetTranslation(FVector(RawTrack.PosKeys[0]));
-					}
-
-					if (RawTrack.RotKeys.IsValidIndex(KeyIndex))
-					{
-						Transform.SetRotation(FQuat(RawTrack.RotKeys[KeyIndex]));
-					}
-					else if (RawTrack.RotKeys.Num() > 0)
-					{
-						Transform.SetRotation(FQuat(RawTrack.RotKeys[0]));
-					}
-
-					if (RawTrack.ScaleKeys.IsValidIndex(KeyIndex))
-					{
-						Transform.SetScale3D(FVector(RawTrack.ScaleKeys[KeyIndex]));
-					}
-					else if (RawTrack.ScaleKeys.Num() > 0)
-					{
-						Transform.SetScale3D(FVector(RawTrack.ScaleKeys[0]));
-					}
-
+					const FTransform Transform = AnimSequence->GetDataModelInterface()->GetBoneTrackTransform(RootBoneName, FFrameNumber(KeyIndex));
 					FFrameNumber AnimationFrame = (AnimSequence->GetTimeAtFrame(KeyIndex) * TickResolution).FloorToFrame();
 					FTransform Total = InvComponentTransform * Transform * Relative;
 					AnimationKeys.Add(Total, StartTime + AnimationFrame);
diff --git a/Engine/Plugins/VirtualProduction/Takes/Source/TakeTrackRecorders/Private/MovieSceneAnimationTrackRecorder.cpp b/Engine/Plugins/VirtualProduction/Takes/Source/TakeTrackRecorders/Private/MovieSceneAnimationTrackRecorder.cpp
--- a/Engine/Plugins/VirtualProduction/Takes/Source/TakeTrackRecorders/Private/MovieSceneAnimationTrackRecorder.cpp
+++ b/Engine/Plugins/VirtualProduction/Takes/Source/TakeTrackRecorders/Private/MovieSceneAnimationTrackRecorder.cpp
@@ -409,7 +409,7 @@ bool UMovieSceneAnimationTrackRecorder::LoadRecordedFile(const FString& FileName
 									int32 MaxNumberOfKeys = 0;
 									for (int32 TrackIndex = 0; TrackIndex < Header.AnimationTrackNames.Num(); ++TrackIndex)
 									{
-										Controller.AddBoneTrack(Header.AnimationTrackNames[TrackIndex]);
+										Controller.AddBoneCurve(Header.AnimationTrackNames[TrackIndex]);
 
 										TArray<FVector3f> PosKeys;
 										TArray<FQuat4f> RotKeys;
diff --git a/Engine/Source/Developer/AnimationDataController/Private/AnimDataController.cpp b/Engine/Source/Developer/AnimationDataController/Private/AnimDataController.cpp
--- a/Engine/Source/Developer/AnimationDataController/Private/AnimDataController.cpp
+++ b/Engine/Source/Developer/AnimationDataController/Private/AnimDataController.cpp
@@ -1279,15 +1279,15 @@ void UAnimDataController::ResizePlayLength(float Length, float T0, float T1, boo
 	ResizeNumberOfFrames(ConvertSecondsToFrameNumber(Length), ConvertSecondsToFrameNumber(T0), ConvertSecondsToFrameNumber(T1), bShouldTransact);
 }
 
-int32 UAnimDataController::AddBoneTrack(FName BoneName, bool bShouldTransact /*= true*/)
+bool UAnimDataController::AddBoneCurve(FName BoneName, bool bShouldTransact /*= true*/)
 {
 	if (!ModelInterface->GetAnimationSequence())
 	{
-		return INDEX_NONE;
+		return false;
 	}
 
 	FTransaction Transaction = ConditionalTransaction(LOCTEXT("AddBoneTrack", "Adding Animation Data Track"), bShouldTransact);
-	return InsertBoneTrack(BoneName, INDEX_NONE, bShouldTransact);
+	return InsertBoneTrack(BoneName, INDEX_NONE, bShouldTransact) != INDEX_NONE;
 }
 
 int32 UAnimDataController::InsertBoneTrack(FName BoneName, int32 DesiredIndex, bool bShouldTransact /*= true*/)
@@ -1343,7 +1343,6 @@ int32 UAnimDataController::InsertBoneTrack(FName BoneName, int32 DesiredIndex, b
 
 				FAnimationTrackAddedPayload Payload;
 				Payload.Name = BoneName;
-				Payload.TrackIndex = InsertIndex;
 
 				Model->GetNotifier().Notify<FAnimationTrackAddedPayload>(EAnimDataModelNotifyType::TrackAdded, Payload);
 				ConditionalAction<UE::Anim::FRemoveTrackAction>(bShouldTransact, BoneName);
@@ -1379,7 +1378,10 @@ bool UAnimDataController::RemoveBoneTrack(FName BoneName, bool bShouldTransact /
 
 		ensure(TrackIndex != INDEX_NONE);
 
-		ConditionalAction<UE::Anim::FAddTrackAction>(bShouldTransact,*ExistingTrackPtr);
+		TArray<FTransform> BoneTransforms;
+		Model->GetBoneTrackTransforms(BoneName, BoneTransforms);
+
+		ConditionalAction<UE::Anim::FAddTrackAction>(bShouldTransact, BoneName, MoveTemp(BoneTransforms));
 		Model->BoneAnimationTracks.RemoveAt(TrackIndex);
 
 		FAnimationTrackRemovedPayload Payload;
@@ -1439,7 +1441,9 @@ bool UAnimDataController::SetBoneTrackKeys(FName BoneName, const TArray<FVector>
 			{
 				FTransaction Transaction = ConditionalTransaction(LOCTEXT("SetTrackKeysTransaction", "Setting Animation Data Track keys"), bShouldTransact);
 
-				ConditionalAction<UE::Anim::FSetTrackKeysAction>(bShouldTransact, *TrackPtr);
+				TArray<FTransform> BoneTransforms;
+				Model->GetBoneTrackTransforms(BoneName, BoneTransforms);
+				ConditionalAction<UE::Anim::FSetTrackKeysAction>(bShouldTransact, BoneName, BoneTransforms);
 
 				TrackPtr->InternalTrackData.PosKeys.SetNum(MaxNumKeys);
 				TrackPtr->InternalTrackData.ScaleKeys.SetNum(MaxNumKeys);
@@ -1498,7 +1502,9 @@ bool UAnimDataController::SetBoneTrackKeys(FName BoneName, const TArray<FVector3
 		{
 			if (FBoneAnimationTrack* TrackPtr = Model->FindMutableBoneTrackByName(BoneName))
 			{
-				ConditionalAction<UE::Anim::FSetTrackKeysAction>(bShouldTransact, *TrackPtr);
+				TArray<FTransform> BoneTransforms;
+				Model->GetBoneTrackTransforms(BoneName, BoneTransforms);
+				ConditionalAction<UE::Anim::FSetTrackKeysAction>(bShouldTransact, BoneName, BoneTransforms);
 
 				TrackPtr->InternalTrackData.PosKeys = PositionalKeys;
 				TrackPtr->InternalTrackData.RotKeys = RotationalKeys;
@@ -1584,7 +1590,9 @@ bool UAnimDataController::UpdateBoneTrackKeys(FName BoneName, const FInt32Range&
 
 						FTransaction Transaction = ConditionalTransaction(LOCTEXT("SetTrackKeysRangeTransaction", "Setting Animation Data Track keys"), bShouldTransact);
 
-						ConditionalAction<UE::Anim::FSetTrackKeysAction>(bShouldTransact, *TrackPtr);
+						TArray<FTransform> BoneTransforms;
+						Model->GetBoneTrackTransforms(BoneName, BoneTransforms);
+						ConditionalAction<UE::Anim::FSetTrackKeysAction>(bShouldTransact, BoneName, BoneTransforms);
 
 						int32 KeyIndex = 0;
 						for (int32 FrameIndex = RangeMin; FrameIndex < RangeMax; ++FrameIndex, ++KeyIndex)
@@ -1667,7 +1675,9 @@ bool UAnimDataController::UpdateBoneTrackKeys(FName BoneName, const FInt32Range&
 
 						FTransaction Transaction = ConditionalTransaction(LOCTEXT("SetTrackKeysTransaction", "Setting Animation Data Track keys"), bShouldTransact);
 
-						ConditionalAction<UE::Anim::FSetTrackKeysAction>(bShouldTransact, *TrackPtr);
+						TArray<FTransform> BoneTransforms;
+						Model->GetBoneTrackTransforms(BoneName, BoneTransforms);
+						ConditionalAction<UE::Anim::FSetTrackKeysAction>(bShouldTransact, BoneName, BoneTransforms);
 
 						int32 KeyIndex = 0;
 						for (int32 FrameIndex = RangeMin; FrameIndex < RangeMax; ++FrameIndex, ++KeyIndex)
@@ -2123,7 +2133,9 @@ void UAnimDataController::UpdateWithSkeleton(USkeleton* TargetSkeleton, bool bSh
 
 void UAnimDataController::PopulateWithExistingModel(TScriptInterface<IAnimationDataModel> InModel)
 {
+	PRAGMA_DISABLE_DEPRECATION_WARNINGS
 	Model->BoneAnimationTracks = InModel->GetBoneAnimationTracks();
+	PRAGMA_ENABLE_DEPRECATION_WARNINGS
 	Model->FrameRate = InModel->GetFrameRate();
 	Model->NumberOfFrames = InModel->GetNumberOfFrames();
 	Model->NumberOfKeys = InModel->GetNumberOfFrames() + 1;
diff --git a/Engine/Source/Developer/AnimationDataController/Private/AnimDataControllerActions.cpp b/Engine/Source/Developer/AnimationDataController/Private/AnimDataControllerActions.cpp
--- a/Engine/Source/Developer/AnimationDataController/Private/AnimDataControllerActions.cpp
+++ b/Engine/Source/Developer/AnimationDataController/Private/AnimDataControllerActions.cpp
@@ -52,23 +52,37 @@ FString FCloseBracketAction::ToStringInternal() const
 	return TEXT("Closing Bracket");
 }
 
-FAddTrackAction::FAddTrackAction(const FBoneAnimationTrack& Track)
+FAddTrackAction::FAddTrackAction(const FName& InName, TArray<FTransform>&& InTransformData)
 {
-	Name = Track.Name;
-	BoneTreeIndex = Track.BoneTreeIndex;
-	Data = Track.InternalTrackData;
+	Name = InName;
+	TransformData = MoveTemp(InTransformData);
 }
 
 TUniquePtr<FChange> FAddTrackAction::ExecuteInternal(IAnimationDataModel* Model, IAnimationDataController* Controller)
 {
-	Controller->AddBoneTrack(Name, false);
-	
-	if (Data.PosKeys.Num() || Data.RotKeys.Num() || Data.ScaleKeys.Num())
+	Controller->AddBoneCurve(Name, false);
+
+	if (TransformData.Num())
 	{
-		Controller->SetBoneTrackKeys(Name, Data.PosKeys, Data.RotKeys, Data.ScaleKeys, false);
+		TArray<FVector3f> PosKeys;
+		TArray<FQuat4f> RotKeys;
+		TArray<FVector3f> ScaleKeys;
+
+		PosKeys.Reserve(TransformData.Num());
+		RotKeys.Reserve(TransformData.Num());
+		ScaleKeys.Reserve(TransformData.Num());	
+		
+		for (const FTransform& Transform : TransformData)
+		{
+			PosKeys.Add(FVector3f(Transform.GetLocation()));
+			RotKeys.Add(FQuat4f(Transform.GetRotation()));
+			ScaleKeys.Add(FVector3f(Transform.GetScale3D()));
+		}
+	
+		Controller->SetBoneTrackKeys(Name, PosKeys, RotKeys, ScaleKeys, false);	
 	}
 	
-	const FBoneAnimationTrack& AddedTrack = Model->GetBoneTrackByName(Name);
+
 	return MakeUnique<FRemoveTrackAction>(Name);
 }
 
@@ -77,16 +91,19 @@ FString FAddTrackAction::ToStringInternal() const
 	return FText::Format(LOCTEXT("AddTrackAction_Description", "Adding animation bone track '{0}'."), FText::FromName(Name)).ToString();
 }
 
-FRemoveTrackAction::FRemoveTrackAction(const FName& TrackName)
+FRemoveTrackAction::FRemoveTrackAction(const FName& InName)
+	: Name(InName)
 {
-	Name = TrackName;
 }
 
 TUniquePtr<FChange> FRemoveTrackAction::ExecuteInternal(IAnimationDataModel* Model, IAnimationDataController* Controller)
 {
-	const FBoneAnimationTrack& Track = Model->GetBoneTrackByName(Name);
+	ensure(Model->IsValidBoneTrackName(Name));
+
+	TArray<FTransform> BoneTransforms;
+	Model->GetBoneTrackTransforms(Name, BoneTransforms);
+	TUniquePtr<FChange> InverseAction = MakeUnique<FAddTrackAction>(Name, MoveTemp(BoneTransforms));
 
-	TUniquePtr<FChange> InverseAction = MakeUnique<FAddTrackAction>(Track);
 	Controller->RemoveBoneTrack(Name, false);
 	return InverseAction;
 }
@@ -96,19 +113,34 @@ FString FRemoveTrackAction::ToStringInternal() const
 	return FText::Format(LOCTEXT("RemoveTrackAction_Description", "Removing animation bone Track '{0}'."), FText::FromName(Name)).ToString();
 }
 
-FSetTrackKeysAction::FSetTrackKeysAction(const FBoneAnimationTrack& Track)
+FSetTrackKeysAction::FSetTrackKeysAction(const FName& InName, TArray<FTransform>& InTransformData)
 {
-	Name = Track.Name;
-
-	TrackData = Track.InternalTrackData;
+	Name = InName;
+	TransformData = MoveTemp(InTransformData);
 }
 
 TUniquePtr<FChange> FSetTrackKeysAction::ExecuteInternal(IAnimationDataModel* Model, IAnimationDataController* Controller)
 {
-	const FBoneAnimationTrack& Track = Model->GetBoneTrackByName(Name);
-	TUniquePtr<FChange> InverseAction = MakeUnique<FSetTrackKeysAction>(Track);
+	TArray<FTransform> CurrentBoneTransforms;
+	Model->GetBoneTrackTransforms(Name, CurrentBoneTransforms);
+	TUniquePtr<FChange> InverseAction = MakeUnique<FSetTrackKeysAction>(Name, CurrentBoneTransforms);
+
+	TArray<FVector3f> PosKeys;
+	TArray<FQuat4f> RotKeys;
+	TArray<FVector3f> ScaleKeys;
 
-	Controller->SetBoneTrackKeys(Name, TrackData.PosKeys, TrackData.RotKeys, TrackData.ScaleKeys, false);
+	PosKeys.Reserve(TransformData.Num());
+	RotKeys.Reserve(TransformData.Num());
+	ScaleKeys.Reserve(TransformData.Num());	
+	
+	for (const FTransform& Transform : TransformData)
+	{
+		PosKeys.Add(FVector3f(Transform.GetLocation()));
+		RotKeys.Add(FQuat4f(Transform.GetRotation()));
+		ScaleKeys.Add(FVector3f(Transform.GetScale3D()));
+	}
+	
+	Controller->SetBoneTrackKeys(Name, PosKeys, RotKeys, ScaleKeys, false);
 	return InverseAction;
 }
 
diff --git a/Engine/Source/Developer/AnimationDataController/Public/AnimDataController.h b/Engine/Source/Developer/AnimationDataController/Public/AnimDataController.h
--- a/Engine/Source/Developer/AnimationDataController/Public/AnimDataController.h
+++ b/Engine/Source/Developer/AnimationDataController/Public/AnimDataController.h
@@ -41,7 +41,7 @@ public:
 	virtual void SetPlayLength(float Length, bool bShouldTransact = true) override;
 	virtual void ResizePlayLength(float NewLength, float T0, float T1, bool bShouldTransact = true) override;
 	virtual void Resize(float Length, float T0, float T1, bool bShouldTransact = true) override;
-	virtual int32 AddBoneTrack(FName BoneName, bool bShouldTransact = true) override;
+	virtual bool AddBoneCurve(FName BoneName, bool bShouldTransact = true) override;
 	virtual int32 InsertBoneTrack(FName BoneName, int32 DesiredIndex, bool bShouldTransact = true) override;
 	virtual bool RemoveBoneTrack(FName BoneName, bool bShouldTransact = true) override;
 	virtual void RemoveAllBoneTracks(bool bShouldTransact = true) override;
diff --git a/Engine/Source/Developer/AnimationDataController/Public/AnimDataControllerActions.h b/Engine/Source/Developer/AnimationDataController/Public/AnimDataControllerActions.h
--- a/Engine/Source/Developer/AnimationDataController/Public/AnimDataControllerActions.h
+++ b/Engine/Source/Developer/AnimationDataController/Public/AnimDataControllerActions.h
@@ -68,7 +68,7 @@ protected:
 class ANIMATIONDATACONTROLLER_API FAddTrackAction : public FAnimDataBaseAction
 {
 public:
-	explicit FAddTrackAction(const FBoneAnimationTrack& Track);
+	explicit FAddTrackAction(const FName& InName, TArray<FTransform>&& InTransformData);
 	virtual ~FAddTrackAction() {}
 protected:
 	FAddTrackAction() {}
@@ -77,14 +77,13 @@ protected:
 
 protected:
 	FName Name;
-	int32 BoneTreeIndex;
-	FRawAnimSequenceTrack Data;
+	TArray<FTransform> TransformData;
 };
 
 class ANIMATIONDATACONTROLLER_API FRemoveTrackAction : public FAnimDataBaseAction
 {
 public:
-	explicit FRemoveTrackAction(const FName& TrackName);
+	explicit FRemoveTrackAction(const FName& InName);
 	virtual ~FRemoveTrackAction() {}
 protected:
 	FRemoveTrackAction() {}
@@ -98,7 +97,7 @@ protected:
 class ANIMATIONDATACONTROLLER_API FSetTrackKeysAction : public FAnimDataBaseAction
 {
 public:
-	explicit FSetTrackKeysAction(const FBoneAnimationTrack& Track);
+	explicit FSetTrackKeysAction(const FName& InName, TArray<FTransform>& InTransformData);
 	virtual ~FSetTrackKeysAction() {}
 protected:
 	FSetTrackKeysAction() {}
@@ -107,9 +106,7 @@ protected:
 
 protected:
 	FName Name;
-	int32 TrackIndex;
-
-	FRawAnimSequenceTrack TrackData;
+	TArray<FTransform> TransformData;
 };
 
 class ANIMATIONDATACONTROLLER_API FResizePlayLengthInFramesAction : public FAnimDataBaseAction
diff --git a/Engine/Source/Developer/AssetTools/Private/AssetTypeActions/AssetTypeActions_AnimBoneCompressionSettings.cpp b/Engine/Source/Developer/AssetTools/Private/AssetTypeActions/AssetTypeActions_AnimBoneCompressionSettings.cpp
--- a/Engine/Source/Developer/AssetTools/Private/AssetTypeActions/AssetTypeActions_AnimBoneCompressionSettings.cpp
+++ b/Engine/Source/Developer/AssetTools/Private/AssetTypeActions/AssetTypeActions_AnimBoneCompressionSettings.cpp
@@ -108,7 +108,7 @@ void FAssetTypeActions_AnimBoneCompressionSettings::ExecuteCompression(TWeakObje
 	for (UAnimSequence* AnimSeq : AnimSeqsToRecompress)
 	{
 		LoadingAnimSlowTask.EnterProgressFrame();
-		AnimSeq->RequestSyncAnimRecompression(false);
+		AnimSeq->CacheDerivedDataForCurrentPlatform();
 	}
 }
 
diff --git a/Engine/Source/Developer/AssetTools/Private/AssetTypeActions/AssetTypeActions_AnimCurveCompressionSettings.cpp b/Engine/Source/Developer/AssetTools/Private/AssetTypeActions/AssetTypeActions_AnimCurveCompressionSettings.cpp
--- a/Engine/Source/Developer/AssetTools/Private/AssetTypeActions/AssetTypeActions_AnimCurveCompressionSettings.cpp
+++ b/Engine/Source/Developer/AssetTools/Private/AssetTypeActions/AssetTypeActions_AnimCurveCompressionSettings.cpp
@@ -107,7 +107,7 @@ void FAssetTypeActions_AnimCurveCompressionSettings::ExecuteCompression(TWeakObj
 	for (UAnimSequence* AnimSeq : AnimSeqsToRecompress)
 	{
 		LoadingAnimSlowTask.EnterProgressFrame();
-		AnimSeq->RequestSyncAnimRecompression(false);
+		AnimSeq->CacheDerivedDataForCurrentPlatform();
 	}
 }
 
diff --git a/Engine/Source/Editor/AnimationBlueprintLibrary/Private/AnimationBlueprintLibrary.cpp b/Engine/Source/Editor/AnimationBlueprintLibrary/Private/AnimationBlueprintLibrary.cpp
--- a/Engine/Source/Editor/AnimationBlueprintLibrary/Private/AnimationBlueprintLibrary.cpp
+++ b/Engine/Source/Editor/AnimationBlueprintLibrary/Private/AnimationBlueprintLibrary.cpp
@@ -163,93 +163,10 @@ void UAnimationBlueprintLibrary::GetAnimationCurveNames(const UAnimSequence* Ani
 	}
 }
 
-void UAnimationBlueprintLibrary::GetRawTrackPositionData(const UAnimSequenceBase* AnimationSequenceBase, const FName TrackName, TArray<FVector>& PositionData)
+const FRawAnimSequenceTrack& UAnimationBlueprintLibrary::GetRawAnimationTrackByName(const UAnimSequenceBase* AnimationSequenceBase, const FName TrackName)
 {
-	PositionData.Empty();
-	if (IsValidRawAnimationTrackName(AnimationSequenceBase, TrackName))
-	{
-		const FRawAnimSequenceTrack& RawTrack =	GetRawAnimationTrackByName(AnimationSequenceBase, TrackName);
-		PositionData.Append(RawTrack.PosKeys);
-	}
-	else
-	{
-		const FString AnimSequenceName = (AnimationSequenceBase != nullptr) ? AnimationSequenceBase->GetName() : "Invalid Animation sequence";
-		UE_LOG(LogAnimationBlueprintLibrary, Warning, TEXT("Raw Animation Track name %s does not exist in Animation Sequence %s"), *TrackName.ToString(), *AnimSequenceName );
-	}	
-}
-
-void UAnimationBlueprintLibrary::GetRawTrackRotationData(const UAnimSequenceBase* AnimationSequenceBase, const FName TrackName, TArray<FQuat>& RotationData)
-{
-	RotationData.Empty();
-	if (IsValidRawAnimationTrackName(AnimationSequenceBase, TrackName))
-	{
-		const FRawAnimSequenceTrack& RawTrack = GetRawAnimationTrackByName(AnimationSequenceBase, TrackName);
-		RotationData.Append(RawTrack.RotKeys);
-	}
-	else
-	{	
-		const FString AnimSequenceName = (AnimationSequenceBase != nullptr) ? AnimationSequenceBase->GetName() : "Invalid Animation sequence";
-		UE_LOG(LogAnimationBlueprintLibrary, Warning, TEXT("Raw Animation Track name %s does not exist in Animation Sequence %s"), *TrackName.ToString(), *AnimSequenceName);
-	}
-}
-
-void UAnimationBlueprintLibrary::GetRawTrackScaleData(const UAnimSequenceBase* AnimationSequenceBase, const FName TrackName, TArray<FVector>& ScaleData)
-{
-	ScaleData.Empty();
-	if (IsValidRawAnimationTrackName(AnimationSequenceBase, TrackName))
-	{
-		const FRawAnimSequenceTrack& RawTrack = GetRawAnimationTrackByName(AnimationSequenceBase, TrackName);
-		ScaleData.Append(RawTrack.ScaleKeys);
-	}
-	else
-	{
-		const FString AnimSequenceName = (AnimationSequenceBase != nullptr) ? AnimationSequenceBase->GetName() : "Invalid Animation sequence";
-		UE_LOG(LogAnimationBlueprintLibrary, Warning, TEXT("Raw Animation Track name %s does not exist in Animation Sequence %s"), *TrackName.ToString(), *AnimSequenceName);
-	}
-}
-
-void UAnimationBlueprintLibrary::GetRawTrackData(const UAnimSequenceBase* AnimationSequenceBase, const FName TrackName, TArray<FVector>& PositionKeys, TArray<FQuat>& RotationKeys, TArray<FVector>& ScalingKeys)
-{
-	PositionKeys.Empty();
-	RotationKeys.Empty();
-	ScalingKeys.Empty();
-	if (IsValidRawAnimationTrackName(AnimationSequenceBase, TrackName))
-	{		
-		const FRawAnimSequenceTrack& RawTrack = GetRawAnimationTrackByName(AnimationSequenceBase, TrackName);
-		PositionKeys.Append(RawTrack.PosKeys);
-		RotationKeys.Append(RawTrack.RotKeys);
-		ScalingKeys.Append(RawTrack.ScaleKeys);
-	}
-	else
-	{
-		const FString AnimSequenceName = (AnimationSequenceBase != nullptr) ? AnimationSequenceBase->GetName() : "Invalid Animation sequence";
-		UE_LOG(LogAnimationBlueprintLibrary, Warning, TEXT("Raw Animation Track name %s does not exist in Animation Sequence %s"), *TrackName.ToString(), *AnimSequenceName);
-	}
-}
-
-bool UAnimationBlueprintLibrary::IsValidRawAnimationTrackName(const UAnimSequenceBase* AnimationSequenceBase, const FName TrackName)
-{
-	bool bValidName = false;
-
-	if (AnimationSequenceBase)
-	{
-		const int32 TrackIndex = AnimationSequenceBase->GetDataModel()->GetBoneTrackIndexByName(TrackName);
-		bValidName = (TrackIndex != INDEX_NONE);
-	}
-	else
-	{
-		UE_LOG(LogAnimationBlueprintLibrary, Warning, TEXT("Invalid Animation Sequence supplied for IsValidRawAnimationTrackName"));
-	}
-
-	return bValidName;
-}
-
-const FRawAnimSequenceTrack&  UAnimationBlueprintLibrary::GetRawAnimationTrackByName(const UAnimSequenceBase* AnimationSequenceBase, const FName TrackName)
-{
-	checkf(AnimationSequenceBase, TEXT("Invalid Animation Sequence supplied for GetRawAnimationTrackByName"));
-
-	const FBoneAnimationTrack& AnimationTrack = AnimationSequenceBase->GetDataModel()->GetBoneTrackByName(TrackName);
-	return AnimationTrack.InternalTrackData;
+	static FRawAnimSequenceTrack TempTrack;
+	return TempTrack;
 }
 
 void UAnimationBlueprintLibrary::GetBoneCompressionSettings(const UAnimSequence* AnimationSequence, UAnimBoneCompressionSettings*& CompressionSettings)
@@ -1889,116 +1806,6 @@ bool UAnimationBlueprintLibrary::ContainsMetaDataOfClass(const UAnimationAsset*
 	return bContainsMetaData;	
 }
 
-void UAnimationBlueprintLibrary::GetBonePoseForTime(const UAnimSequenceBase* AnimationSequenceBase, FName BoneName, float Time, bool bExtractRootMotion, FTransform& Pose)
-{
-	Pose.SetIdentity();
-	if (AnimationSequenceBase)
-	{
-		TArray<FName> BoneNameArray;
-		TArray<FTransform> PoseArray;
-		BoneNameArray.Add(BoneName);
-		GetBonePosesForTime(AnimationSequenceBase, BoneNameArray, Time, bExtractRootMotion, PoseArray);
-		Pose = PoseArray[0];
-	}
-	else
-	{
-		UE_LOG(LogAnimationBlueprintLibrary, Warning, TEXT("Invalid Animation Sequence supplied for GetBonePoseForTime"));
-	}
-}
-
-void UAnimationBlueprintLibrary::GetBonePoseForFrame(const UAnimSequenceBase* AnimationSequenceBase, FName BoneName, int32 Frame, bool bExtractRootMotion, FTransform& Pose)
-{
-	Pose.SetIdentity();
-	if (AnimationSequenceBase)
-	{
-		GetBonePoseForTime(AnimationSequenceBase, BoneName, GetTimeAtFrameInternal(AnimationSequenceBase, Frame), bExtractRootMotion, Pose);
-	}
-	else
-	{
-		UE_LOG(LogAnimationBlueprintLibrary, Warning, TEXT("Invalid Animation Sequence supplied for GetBonePoseForFrame"));
-	}
-}
-
-void UAnimationBlueprintLibrary::GetBonePosesForTime(const UAnimSequenceBase* AnimationSequenceBase, TArray<FName> BoneNames, float Time, bool bExtractRootMotion, TArray<FTransform>& Poses, const USkeletalMesh* PreviewMesh /*= nullptr*/)
-{
-	Poses.Empty(BoneNames.Num());
-	if (AnimationSequenceBase && AnimationSequenceBase->GetSkeleton())
-	{
-		Poses.AddDefaulted(BoneNames.Num());
-
-		// Need this for FCompactPose
-		FMemMark Mark(FMemStack::Get());
-
-		const FReferenceSkeleton& RefSkeleton = (PreviewMesh)? PreviewMesh->GetRefSkeleton() : AnimationSequenceBase->GetSkeleton()->GetReferenceSkeleton();
-
-		if (IsValidTimeInternal(AnimationSequenceBase, Time))
-		{
-			if (BoneNames.Num())
-			{
-				for (int32 BoneNameIndex = 0; BoneNameIndex < BoneNames.Num(); ++BoneNameIndex)
-				{
-					const FName& BoneName = BoneNames[BoneNameIndex];
-					
-					FTransform& Transform = Poses[BoneNameIndex];
-					if (IsValidRawAnimationTrackName(AnimationSequenceBase, BoneName))
-					{
-						const EAnimInterpolationType InterpolationType = [AnimationSequenceBase]() -> EAnimInterpolationType
-					{
-							if (const UAnimSequence* AnimationSequence = Cast<const UAnimSequence>(AnimationSequenceBase))
-							{
-								return AnimationSequence->Interpolation;
-							}
-
-							return EAnimInterpolationType::Linear;
-						}();
-						UE::Anim::GetBoneTransformFromModel(AnimationSequenceBase->GetDataModel(), Transform, AnimationSequenceBase->GetDataModel()->GetBoneTrackIndexByName(BoneName), Time, InterpolationType);
-					}
-					else
-					{
-						
-						// otherwise, get ref pose if exists
-						const int32 BoneIndex = RefSkeleton.FindBoneIndex(BoneName);
-						if (BoneIndex != INDEX_NONE)
-						{
-							Transform = RefSkeleton.GetRefBonePose()[BoneIndex];
-						}
-						else
-						{
-							UE_LOG(LogAnimationBlueprintLibrary, Warning, TEXT("Invalid bone name %s for Animation Sequence %s supplied for GetBonePosesForTime"), *BoneName.ToString(), *AnimationSequenceBase->GetName());
-							Transform = FTransform::Identity;
-						}
-					}
-				}
-			}
-			else
-			{
-				UE_LOG(LogAnimationBlueprintLibrary, Error, TEXT("Invalid or no bone names specified to retrieve poses given Animation Sequence %s in GetBonePosesForTime"), *AnimationSequenceBase->GetName());
-			}			
-		}
-		else
-		{
-			UE_LOG(LogAnimationBlueprintLibrary, Warning, TEXT("Invalid time value %f for Animation Sequence %s supplied for GetBonePosesForTime"), Time, *AnimationSequenceBase->GetName());
-		}
-	}
-	else
-	{
-		UE_LOG(LogAnimationBlueprintLibrary, Warning, TEXT("Invalid Animation Sequence supplied for GetBonePosesForTime"));
-	}
-}
-
-void UAnimationBlueprintLibrary::GetBonePosesForFrame(const UAnimSequenceBase* AnimationSequenceBase, TArray<FName> BoneNames, int32 Frame, bool bExtractRootMotion, TArray<FTransform>& Poses, const USkeletalMesh* PreviewMesh /*= nullptr*/)
-{
-	Poses.Empty(BoneNames.Num());
-	if (AnimationSequenceBase)
-	{
-		GetBonePosesForTime(AnimationSequenceBase, BoneNames, GetTimeAtFrameInternal(AnimationSequenceBase, Frame), bExtractRootMotion, Poses, PreviewMesh);
-	}
-	else
-	{
-		UE_LOG(LogAnimationBlueprintLibrary, Warning, TEXT("Invalid Animation Sequence supplied for GetBonePosesForFrame"));
-	}
-}
-
 void UAnimationBlueprintLibrary::AddVirtualBone(const UAnimSequence* AnimationSequence, FName SourceBoneName, FName TargetBoneName, FName& VirtualBoneName)
 {
 	if (AnimationSequence)
@@ -2237,15 +2044,12 @@ bool UAnimationBlueprintLibrary::EvaluateRootBoneTimecodeAttributesAtTime(const
 		return false;
 	}
 
-	const int32 RootBoneTrackIndex = 0;
-	const FBoneAnimationTrack* RootBoneTrack = AnimDataModel->FindBoneTrackByIndex(RootBoneTrackIndex);
-	if (!RootBoneTrack)
+	const FName RootBoneName = AnimationSequenceBase->GetSkeleton()->GetReferenceSkeleton().GetBoneName(0);
+	if (!AnimDataModel->IsValidBoneTrackName(RootBoneName))
 	{
 		return false;
 	}
 
-	const FName& RootBoneName = RootBoneTrack->Name;
-
 	TArray<const FAnimatedBoneAttribute*> RootBoneAttributes;
 	AnimDataModel->GetAttributesForBone(RootBoneName, RootBoneAttributes);
 
@@ -2426,22 +2230,19 @@ bool UAnimationBlueprintLibrary::EvaluateRootBoneTimecodeSubframeAttributeAtTime
 		return false;
 	}
 
-	const int32 RootBoneIndex = 0;
-	const FBoneAnimationTrack* RootBoneTrack = AnimDataModel->FindBoneTrackByIndex(RootBoneIndex);
-	if (!RootBoneTrack)
+	const FName RootBoneName = AnimationSequenceBase->GetSkeleton()->GetReferenceSkeleton().GetBoneName(0);
+	if (!AnimDataModel->IsValidBoneTrackName(RootBoneName))
 	{
 		return false;
 	}
 
-	const FName& RootBoneName = RootBoneTrack->Name;
-
 	FName TCSubframeAttrName(TEXT("TCSubframe"));
 	if (const UAnimationSettings* AnimationSettings = UAnimationSettings::Get())
 	{
 		TCSubframeAttrName = AnimationSettings->BoneTimecodeCustomAttributeNameSettings.SubframeAttributeName;
 	}
 
-	FAnimationAttributeIdentifier SubframeAttributeIdentifier(TCSubframeAttrName, RootBoneIndex, RootBoneName, FFloatAnimationAttribute::StaticStruct());
+	FAnimationAttributeIdentifier SubframeAttributeIdentifier(TCSubframeAttrName, 0, RootBoneName, FFloatAnimationAttribute::StaticStruct());
 	const FAnimatedBoneAttribute* RootBoneSubframeAttribute = AnimDataModel->FindAttribute(SubframeAttributeIdentifier);
 	if (!RootBoneSubframeAttribute)
 	{
diff --git a/Engine/Source/Editor/AnimationBlueprintLibrary/Public/AnimationBlueprintLibrary.h b/Engine/Source/Editor/AnimationBlueprintLibrary/Public/AnimationBlueprintLibrary.h
--- a/Engine/Source/Editor/AnimationBlueprintLibrary/Public/AnimationBlueprintLibrary.h
+++ b/Engine/Source/Editor/AnimationBlueprintLibrary/Public/AnimationBlueprintLibrary.h
@@ -86,25 +86,30 @@ public:
 	static void GetAnimationCurveNames(const UAnimSequence* AnimationSequence, ERawCurveTrackTypes CurveType, TArray<FName>& CurveNames);
 
 	/** Retrieves the Raw Translation Animation Data for the given Animation Track Name and Animation Sequence */
+	UE_DEPRECATED(5.2, "GetRawTrackPositionData has been deprecated, use AnimationModel interface instead")
 	UFUNCTION(BlueprintPure, Category = "AnimationBlueprintLibrary|RawTrackData")
-	static void GetRawTrackPositionData(const UAnimSequenceBase* AnimationSequenceBase, const FName TrackName, TArray<FVector>& PositionData);
+	static void GetRawTrackPositionData(const UAnimSequenceBase* AnimationSequenceBase, const FName TrackName, TArray<FVector>& PositionData) {}
 
 	/** Retrieves the Raw Rotation Animation Data for the given Animation Track Name and Animation Sequence */
+	UE_DEPRECATED(5.2, "GetRawTrackRotationData has been deprecated, use AnimationModel interface instead")
 	UFUNCTION(BlueprintPure, Category = "AnimationBlueprintLibrary|RawTrackData")
-	static void GetRawTrackRotationData(const UAnimSequenceBase* AnimationSequenceBase, const FName TrackName, TArray<FQuat>& RotationData );
+	static void GetRawTrackRotationData(const UAnimSequenceBase* AnimationSequenceBase, const FName TrackName, TArray<FQuat>& RotationData ) {}
 
-	/** Retrieves the Raw Scale Animation Data for the given Animation Track Name and Animation Sequence */
+	/** Retrieves the Raw Scale Animation Data for the given Animation Track Name and Animation Sequence */	
+	UE_DEPRECATED(5.2, "GetRawTrackScaleData has been deprecated, use AnimationModel interface instead")
 	UFUNCTION(BlueprintPure, Category = "AnimationBlueprintLibrary|RawTrackData")
-	static void GetRawTrackScaleData(const UAnimSequenceBase* AnimationSequenceBase, const FName TrackName, TArray<FVector>& ScaleData);
+	static void GetRawTrackScaleData(const UAnimSequenceBase* AnimationSequenceBase, const FName TrackName, TArray<FVector>& ScaleData) {}
 
 	/** Retrieves the Raw Animation Data for the given Animation Track Name and Animation Sequence */
+	UE_DEPRECATED(5.2, "GetRawTrackScaleData has been deprecated, use AnimationModel interface instead")
 	UFUNCTION(BlueprintPure, Category = "AnimationBlueprintLibrary|RawTrackData")
-	static void GetRawTrackData(const UAnimSequenceBase* AnimationSequenceBase, const FName TrackName, TArray<FVector>& PositionKeys,TArray<FQuat>& RotationKeys, TArray<FVector>& ScalingKeys);
+	static void GetRawTrackData(const UAnimSequenceBase* AnimationSequenceBase, const FName TrackName, TArray<FVector>& PositionKeys,TArray<FQuat>& RotationKeys, TArray<FVector>& ScalingKeys) {}
 
 	/** Checks whether or not the given Animation Track Name is contained within the Animation Sequence */
 	UFUNCTION(BlueprintCallable, Category = "AnimationBlueprintLibrary|Helpers")
-	static bool IsValidRawAnimationTrackName(const UAnimSequenceBase* AnimationSequenceBase, const FName TrackName);
+	static bool IsValidRawAnimationTrackName(const UAnimSequenceBase* AnimationSequenceBase, const FName TrackName) { return false; }
 
+	UE_DEPRECATED(5.2, "GetRawAnimationTrackByName has been deprecated")
 	static const FRawAnimSequenceTrack& GetRawAnimationTrackByName(const UAnimSequenceBase* AnimationSequenceBase, const FName TrackName);
 
 	// Compression
@@ -436,20 +441,24 @@ public:
 	// Poses
 
 	/** Retrieves Bone Pose data for the given Bone Name at the specified Time from the given Animation Sequence */
+	UE_DEPRECATED(5.2, "GetBonePosesForTime is deprecated, use AnimPose or AnimationDataModel interface directly")
 	UFUNCTION(BlueprintPure, Category = "AnimationBlueprintLibrary|Pose")
-	static void GetBonePoseForTime(const UAnimSequenceBase* AnimationSequenceBase, FName BoneName, float Time, bool bExtractRootMotion, FTransform& Pose);
+	static void GetBonePoseForTime(const UAnimSequenceBase* AnimationSequenceBase, FName BoneName, float Time, bool bExtractRootMotion, FTransform& Pose) {}
 
 	/** Retrieves Bone Pose data for the given Bone Name at the specified Frame from the given Animation Sequence */
+	UE_DEPRECATED(5.2, "GetBonePosesForTime is deprecated, use AnimPose or AnimationDataModel interface directly")
 	UFUNCTION(BlueprintPure, Category = "AnimationBlueprintLibrary|Pose")
-	static void GetBonePoseForFrame(const UAnimSequenceBase* AnimationSequenceBase, FName BoneName, int32 Frame, bool bExtractRootMotion, FTransform& Pose);
+	static void GetBonePoseForFrame(const UAnimSequenceBase* AnimationSequenceBase, FName BoneName, int32 Frame, bool bExtractRootMotion, FTransform& Pose) {}
 
 	/** Retrieves Bone Pose data for the given Bone Names at the specified Time from the given Animation Sequence */
+	UE_DEPRECATED(5.2, "GetBonePosesForTime is deprecated, use AnimPose or AnimationDataModel interface directly")
 	UFUNCTION(BlueprintPure, Category = "AnimationBlueprintLibrary|Pose")
-	static void GetBonePosesForTime(const UAnimSequenceBase* AnimationSequenceBase, TArray<FName> BoneNames, float Time, bool bExtractRootMotion, TArray<FTransform>& Poses, const USkeletalMesh* PreviewMesh = nullptr);
+	static void GetBonePosesForTime(const UAnimSequenceBase* AnimationSequenceBase, TArray<FName> BoneNames, float Time, bool bExtractRootMotion, TArray<FTransform>& Poses, const USkeletalMesh* PreviewMesh = nullptr) {}
 
 	/** Retrieves Bone Pose data for the given Bone Names at the specified Frame from the given Animation Sequence */
+	UE_DEPRECATED(5.2, "GetBonePosesForTime is deprecated, use AnimPose or AnimationDataModel interface directly")
 	UFUNCTION(BlueprintPure, Category = "AnimationBlueprintLibrary|Pose")
-	static void GetBonePosesForFrame(const UAnimSequenceBase* AnimationSequenceBase, TArray<FName> BoneNames, int32 Frame, bool bExtractRootMotion, TArray<FTransform>& Poses, const USkeletalMesh* PreviewMesh = nullptr);
+	static void GetBonePosesForFrame(const UAnimSequenceBase* AnimationSequenceBase, TArray<FName> BoneNames, int32 Frame, bool bExtractRootMotion, TArray<FTransform>& Poses, const USkeletalMesh* PreviewMesh = nullptr) {}
 
 	// Virtual bones
 
diff --git a/Engine/Source/Editor/AnimationModifiers/Private/AnimationModifier.cpp b/Engine/Source/Editor/AnimationModifiers/Private/AnimationModifier.cpp
--- a/Engine/Source/Editor/AnimationModifiers/Private/AnimationModifier.cpp
+++ b/Engine/Source/Editor/AnimationModifiers/Private/AnimationModifier.cpp
@@ -191,7 +191,7 @@ void UAnimationModifier::UpdateCompressedAnimationData()
 {
 	if (CurrentAnimSequence->DoesNeedRecompress())
 	{
-		CurrentAnimSequence->RequestSyncAnimRecompression(false);
+		CurrentAnimSequence->CacheDerivedDataForCurrentPlatform();
 	}
 }
 
diff --git a/Engine/Source/Editor/DetailCustomizations/Private/DetailCustomizations.cpp b/Engine/Source/Editor/DetailCustomizations/Private/DetailCustomizations.cpp
--- a/Engine/Source/Editor/DetailCustomizations/Private/DetailCustomizations.cpp
+++ b/Engine/Source/Editor/DetailCustomizations/Private/DetailCustomizations.cpp
@@ -345,6 +345,7 @@ void FDetailCustomizationsModule::RegisterPropertyTypeCustomizations()
 	RegisterCustomPropertyTypeLayout("PerPlatformInt", FOnGetPropertyTypeCustomizationInstance::CreateStatic(&FPerPlatformPropertyCustomization<FPerPlatformInt>::MakeInstance));
 	RegisterCustomPropertyTypeLayout("PerPlatformFloat", FOnGetPropertyTypeCustomizationInstance::CreateStatic(&FPerPlatformPropertyCustomization<FPerPlatformFloat>::MakeInstance));
 	RegisterCustomPropertyTypeLayout("PerPlatformBool", FOnGetPropertyTypeCustomizationInstance::CreateStatic(&FPerPlatformPropertyCustomization<FPerPlatformBool>::MakeInstance));
+	RegisterCustomPropertyTypeLayout("PerPlatformFrameRate", FOnGetPropertyTypeCustomizationInstance::CreateStatic(&FPerPlatformPropertyCustomization<FPerPlatformFrameRate>::MakeInstance));
 	RegisterCustomPropertyTypeLayout("PerQualityLevelInt", FOnGetPropertyTypeCustomizationInstance::CreateStatic(&FPerQualityLevelPropertyCustomization<FPerQualityLevelInt>::MakeInstance));
 	RegisterCustomPropertyTypeLayout("PerQualityLevelFloat", FOnGetPropertyTypeCustomizationInstance::CreateStatic(&FPerQualityLevelPropertyCustomization<FPerQualityLevelFloat>::MakeInstance));
 	RegisterCustomPropertyTypeLayout("SkeletalMeshOptimizationSettings", FOnGetPropertyTypeCustomizationInstance::CreateStatic(&FSkeletalMeshReductionSettingsDetails::MakeInstance));
diff --git a/Engine/Source/Editor/DetailCustomizations/Private/PerPlatformPropertyCustomization.cpp b/Engine/Source/Editor/DetailCustomizations/Private/PerPlatformPropertyCustomization.cpp
--- a/Engine/Source/Editor/DetailCustomizations/Private/PerPlatformPropertyCustomization.cpp
+++ b/Engine/Source/Editor/DetailCustomizations/Private/PerPlatformPropertyCustomization.cpp
@@ -57,14 +57,17 @@ void FPerPlatformPropertyCustomization<PerPlatformType>::CustomizeChildren(TShar
 	Args.PlatformOverrideNames = PlatformOverrideNames;
 	Args.OnAddPlatformOverride = FOnPlatformOverrideAction::CreateSP(this, &FPerPlatformPropertyCustomization<PerPlatformType>::AddPlatformOverride, StructPropertyHandle);
 	Args.OnRemovePlatformOverride = FOnPlatformOverrideAction::CreateSP(this, &FPerPlatformPropertyCustomization<PerPlatformType>::RemovePlatformOverride, StructPropertyHandle);
-	Args.OnGenerateWidgetForPlatformRow = FOnGenerateWidget::CreateSP(this, &FPerPlatformPropertyCustomization<PerPlatformType>::GetWidget, StructPropertyHandle);
-
+	Args.OnGenerateWidgetForPlatformRow = FOnGenerateWidget::CreateLambda([this, StructPropertyHandle, &StructBuilder](FName PlatformGroupName)
+	{
+		return GetWidget(PlatformGroupName, StructPropertyHandle, StructBuilder);
+	});
+	
 	StructBuilder.AddCustomBuilder(MakeShared<FPerPlatformPropertyCustomNodeBuilder>(MoveTemp(Args)));
 }
 
 
 template<typename PerPlatformType>
-TSharedRef<SWidget> FPerPlatformPropertyCustomization<PerPlatformType>::GetWidget(FName PlatformGroupName, TSharedRef<IPropertyHandle> StructPropertyHandle) const
+TSharedRef<SWidget> FPerPlatformPropertyCustomization<PerPlatformType>::GetWidget(FName PlatformGroupName, TSharedRef<IPropertyHandle> StructPropertyHandle, IDetailChildrenBuilder& StructBuilder) const
 {
 	TSharedPtr<IPropertyHandle>	EditProperty;
 
@@ -107,7 +110,7 @@ TSharedRef<SWidget> FPerPlatformPropertyCustomization<PerPlatformType>::GetWidge
 		// Iterate through source map, setting each key/value pair in the destination
 		for (const auto& It : *SourceMap)
 		{
-				EditProperty->SetInstanceMetaData(*It.Key.ToString(), *It.Value);
+			EditProperty->SetInstanceMetaData(*It.Key.ToString(), *It.Value);
 		}
 
 		// Copy instance metadata as well
@@ -120,6 +123,11 @@ TSharedRef<SWidget> FPerPlatformPropertyCustomization<PerPlatformType>::GetWidge
 
 	if (ensure(EditProperty.IsValid()))
 	{
+		if (EditProperty->GetProperty()->IsA<FStructProperty>())
+		{
+			return StructBuilder.GenerateStructValueWidget(EditProperty->AsShared());
+		}
+		
 		return EditProperty->CreatePropertyValueWidget(false);
 	}
 	
@@ -157,9 +165,9 @@ bool FPerPlatformPropertyCustomization<PerPlatformType>::AddPlatformOverride(FNa
 							KeyProperty->SetValue(PlatformGroupName);
 
 							// Set Value
-							typename PerPlatformType::ValueType DefaultValue;
-							DefaultProperty->GetValue(DefaultValue);
-							ChildProperty->SetValue(DefaultValue);
+							FString PropertyValueString;
+							DefaultProperty->GetValueAsFormattedString(PropertyValueString);
+							ChildProperty->SetValueFromFormattedString(PropertyValueString);
 
 							return true;
 						}
@@ -240,6 +248,7 @@ TSharedRef<IPropertyTypeCustomization> FPerPlatformPropertyCustomization<PerPlat
 template class FPerPlatformPropertyCustomization<FPerPlatformInt>;
 template class FPerPlatformPropertyCustomization<FPerPlatformFloat>;
 template class FPerPlatformPropertyCustomization<FPerPlatformBool>;
+template class FPerPlatformPropertyCustomization<FPerPlatformFrameRate>;
 
 #undef LOCTEXT_NAMESPACE
 
diff --git a/Engine/Source/Editor/DetailCustomizations/Public/PerPlatformPropertyCustomization.h b/Engine/Source/Editor/DetailCustomizations/Public/PerPlatformPropertyCustomization.h
--- a/Engine/Source/Editor/DetailCustomizations/Public/PerPlatformPropertyCustomization.h
+++ b/Engine/Source/Editor/DetailCustomizations/Public/PerPlatformPropertyCustomization.h
@@ -100,7 +100,7 @@ public:
 	static TSharedRef<IPropertyTypeCustomization> MakeInstance();
 
 protected:
-	TSharedRef<SWidget> GetWidget(FName PlatformGroupName, TSharedRef<IPropertyHandle> StructPropertyHandle) const;
+	TSharedRef<SWidget> GetWidget(FName PlatformGroupName, TSharedRef<IPropertyHandle> StructPropertyHandle, IDetailChildrenBuilder& StructBuilder) const;
 	TArray<FName> GetPlatformOverrideNames(TSharedRef<IPropertyHandle> StructPropertyHandle) const;
 	bool AddPlatformOverride(FName PlatformGroupName, TSharedRef<IPropertyHandle> StructPropertyHandle);
 	bool RemovePlatformOverride(FName PlatformGroupName, TSharedRef<IPropertyHandle> StructPropertyHandle);
diff --git a/Engine/Source/Editor/MovieSceneTools/Private/TrackEditors/TransformTrackEditor.cpp b/Engine/Source/Editor/MovieSceneTools/Private/TrackEditors/TransformTrackEditor.cpp
--- a/Engine/Source/Editor/MovieSceneTools/Private/TrackEditors/TransformTrackEditor.cpp
+++ b/Engine/Source/Editor/MovieSceneTools/Private/TrackEditors/TransformTrackEditor.cpp
@@ -33,6 +33,7 @@
 
 #include "Tracks/IMovieSceneTransformOrigin.h"
 #include "IMovieScenePlaybackClient.h"
+#include "Animation/AnimData/AnimDataModel.h"
 
 #include "Editor.h"
 #include "LevelEditorViewport.h"
@@ -1265,41 +1266,17 @@ void F3DTransformTrackEditor::ImportAnimSequenceTransforms(const FAssetData& Ass
 
 			TArray<FTempTransformKey> TempKeys;
 
-			const FBoneAnimationTrack& AnimationTrack = AnimSequence->GetDataModel()->GetBoneTrackByIndex(0);
-			const FRawAnimSequenceTrack& RawTrack = AnimationTrack.InternalTrackData;
+			TArray<FName> BoneTrackNames;
+			AnimSequence->GetDataModelInterface()->GetBoneTrackNames(BoneTrackNames);
 
-			const int32 KeyCount = FMath::Max(FMath::Max(RawTrack.PosKeys.Num(), RawTrack.RotKeys.Num()), RawTrack.ScaleKeys.Num());
-			for(int32 KeyIndex = 0; KeyIndex < KeyCount; KeyIndex++)
+			TArray<FTransform> BoneTransforms;
+			AnimSequence->GetDataModelInterface()->GetBoneTrackTransforms(BoneTrackNames[0], BoneTransforms);
+
+			for(int32 KeyIndex = 0; KeyIndex < BoneTransforms.Num(); KeyIndex++)
 			{
 				FTempTransformKey TempKey;
 				TempKey.Time = AnimSequence->GetTimeAtFrame(KeyIndex);
-
-				if(RawTrack.PosKeys.IsValidIndex(KeyIndex))
-				{
-					TempKey.Transform.SetTranslation(FVector(RawTrack.PosKeys[KeyIndex]));
-				}
-				else if(RawTrack.PosKeys.Num() > 0)
-				{
-					TempKey.Transform.SetTranslation(FVector(RawTrack.PosKeys[0]));
-				}
-				
-				if(RawTrack.RotKeys.IsValidIndex(KeyIndex))
-				{
-					TempKey.Transform.SetRotation(FQuat(RawTrack.RotKeys[KeyIndex]));
-				}
-				else if(RawTrack.RotKeys.Num() > 0)
-				{
-					TempKey.Transform.SetRotation(FQuat(RawTrack.RotKeys[0]));
-				}
-
-				if(RawTrack.ScaleKeys.IsValidIndex(KeyIndex))
-				{
-					TempKey.Transform.SetScale3D(FVector(RawTrack.ScaleKeys[KeyIndex]));
-				}
-				else if(RawTrack.ScaleKeys.Num() > 0)
-				{
-					TempKey.Transform.SetScale3D(FVector(RawTrack.ScaleKeys[0]));
-				}
+				TempKey.Transform = BoneTransforms[KeyIndex];
 
 				// apply component transform if any
 				TempKey.Transform = InvComponentTransform * TempKey.Transform;
diff --git a/Engine/Source/Editor/Persona/Private/AnimTimeline/SAnimTimeline.cpp b/Engine/Source/Editor/Persona/Private/AnimTimeline/SAnimTimeline.cpp
--- a/Engine/Source/Editor/Persona/Private/AnimTimeline/SAnimTimeline.cpp
+++ b/Engine/Source/Editor/Persona/Private/AnimTimeline/SAnimTimeline.cpp
@@ -580,46 +580,56 @@ void SAnimTimeline::OnReZeroAnimSequence(int32 FrameIndex)
 
 		if (PreviewInstance->GetCurrentAsset() && PreviewSkelComp )
 		{
-			UAnimSequence* AnimSequence = Cast<UAnimSequence>( PreviewInstance->GetCurrentAsset() );
-			if( AnimSequence )
+			if(UAnimSequence* AnimSequence = Cast<UAnimSequence>( PreviewInstance->GetCurrentAsset()))
 			{
-				const FScopedTransaction Transaction( LOCTEXT("ReZeroAnimation", "ReZero Animation Sequence") );
+				const FScopedTransaction Transaction( LOCTEXT("ReZeroAnimation", "ReZero Animation Sequence"));
+				if (const USkeleton* Skeleton = AnimSequence->GetSkeleton())
+				{
+					const FName RootBoneName = Skeleton->GetReferenceSkeleton().GetBoneName(0);
 
-				//Call modify to restore anim sequence current state
-				AnimSequence->Modify();
+					if(AnimSequence->GetDataModel()->IsValidBoneTrackName(RootBoneName))
+					{
+						TArray<FVector3f> PosKeys;
+						TArray<FQuat4f> RotKeys;
+						TArray<FVector3f> ScaleKeys;
 
-				// As above, animations don't have any idea of hierarchy, so we don't know for sure if track 0 is the root bone's track.
-				const FBoneAnimationTrack& AnimationTrack = AnimSequence->GetDataModel()->GetBoneTrackByIndex(0);
-				FRawAnimSequenceTrack RawTrack = AnimationTrack.InternalTrackData;
+						TArray<FTransform> BoneTransforms;
+						AnimSequence->GetDataModel()->GetBoneTrackTransforms(RootBoneName, BoneTransforms);
 
-				// Find vector that would translate current root bone location onto origin.
-				FVector FrameTransform = FVector::ZeroVector;
-				if (FrameIndex == INDEX_NONE)
-				{
-					// Use current transform
-					FrameTransform = PreviewSkelComp->GetComponentSpaceTransforms()[0].GetLocation();
-				}
-				else if(RawTrack.PosKeys.IsValidIndex(FrameIndex))
-				{
-					// Use transform at frame
-					FrameTransform = (FVector)RawTrack.PosKeys[FrameIndex];
-				}
+						PosKeys.SetNum(BoneTransforms.Num());
+						RotKeys.SetNum(BoneTransforms.Num());
+						ScaleKeys.SetNum(BoneTransforms.Num());
 
-				FVector ApplyTranslation = -1.f * FrameTransform;
+						// Find vector that would translate current root bone location onto origin.
+						FVector FrameTransform = FVector::ZeroVector;
+						if (FrameIndex == INDEX_NONE)
+						{
+							// Use current transform
+							FrameTransform = PreviewSkelComp->GetComponentSpaceTransforms()[0].GetLocation();
+						}
+						else if(BoneTransforms.IsValidIndex(FrameIndex))
+						{
+							// Use transform at frame
+							FrameTransform = BoneTransforms[FrameIndex].GetLocation();
+						}
 
-				// Convert into world space
-				FVector WorldApplyTranslation = PreviewSkelComp->GetComponentTransform().TransformVector(ApplyTranslation);
-				ApplyTranslation = PreviewSkelComp->GetComponentTransform().InverseTransformVector(WorldApplyTranslation);
+						FVector ApplyTranslation = -1.f * FrameTransform;
 
-				for(int32 i=0; i<RawTrack.PosKeys.Num(); i++)
-				{
-					RawTrack.PosKeys[i] += (FVector3f)ApplyTranslation;
-				}
+						// Convert into world space
+						const FVector WorldApplyTranslation = PreviewSkelComp->GetComponentTransform().TransformVector(ApplyTranslation);
+						ApplyTranslation = PreviewSkelComp->GetComponentTransform().InverseTransformVector(WorldApplyTranslation);
 
-				IAnimationDataController& Controller = AnimSequence->GetController();
-				Controller.SetBoneTrackKeys(AnimationTrack.Name, RawTrack.PosKeys, RawTrack.RotKeys, RawTrack.ScaleKeys);
+						for(int32 KeyIndex = 0; KeyIndex < BoneTransforms.Num(); KeyIndex++)
+						{
+							PosKeys[KeyIndex] = FVector3f(BoneTransforms[KeyIndex].GetLocation() + ApplyTranslation);
+							RotKeys[KeyIndex] = FQuat4f(BoneTransforms[KeyIndex].GetRotation());
+							ScaleKeys[KeyIndex] = FVector3f(BoneTransforms[KeyIndex].GetScale3D());
+						}
 
-				AnimSequence->MarkPackageDirty();
+						IAnimationDataController& Controller = AnimSequence->GetController();
+						Controller.SetBoneTrackKeys(RootBoneName, PosKeys, RotKeys, ScaleKeys);
+					}
+				}
 			}
 		}
 	}
diff --git a/Engine/Source/Editor/Persona/Private/SAnimationScrubPanel.cpp b/Engine/Source/Editor/Persona/Private/SAnimationScrubPanel.cpp
--- a/Engine/Source/Editor/Persona/Private/SAnimationScrubPanel.cpp
+++ b/Engine/Source/Editor/Persona/Private/SAnimationScrubPanel.cpp
@@ -552,46 +552,55 @@ void SAnimationScrubPanel::OnReZeroAnimSequence(int32 FrameIndex)
 
 		if (PreviewInstance->GetCurrentAsset() && PreviewSkelComp )
 		{
-			UAnimSequence* AnimSequence = Cast<UAnimSequence>( PreviewInstance->GetCurrentAsset() );
-			if( AnimSequence )
+			if(UAnimSequence* AnimSequence = Cast<UAnimSequence>( PreviewInstance->GetCurrentAsset()))
 			{
-				const FScopedTransaction Transaction( LOCTEXT("ReZeroAnimation", "ReZero Animation Sequence") );
-
-				//Call modify to restore anim sequence current state
-				AnimSequence->Modify();
-
-				// As above, animations don't have any idea of hierarchy, so we don't know for sure if track 0 is the root bone's track.
-				const FBoneAnimationTrack& AnimationTrack = AnimSequence->GetDataModel()->GetBoneTrackByIndex(0);
-				FRawAnimSequenceTrack RawTrack = AnimationTrack.InternalTrackData;
-
-				// Find vector that would translate current root bone location onto origin.
-				FVector FrameTransform = FVector::ZeroVector;
-				if (FrameIndex == INDEX_NONE)
-				{
-					// Use current transform
-					FrameTransform = PreviewSkelComp->GetComponentSpaceTransforms()[0].GetLocation();
-				}
-				else if(RawTrack.PosKeys.IsValidIndex(FrameIndex))
-				{
-					// Use transform at frame
-					FrameTransform = (FVector)RawTrack.PosKeys[FrameIndex];
-				}
-
-				FVector ApplyTranslation = -1.f * FrameTransform;
-
-				// Convert into world space
-				FVector WorldApplyTranslation = PreviewSkelComp->GetComponentTransform().TransformVector(ApplyTranslation);
-				ApplyTranslation = PreviewSkelComp->GetComponentTransform().InverseTransformVector(WorldApplyTranslation);
-
-				for(int32 i=0; i<RawTrack.PosKeys.Num(); i++)
+				if (const USkeleton* Skeleton = AnimSequence->GetSkeleton())
 				{
-					RawTrack.PosKeys[i] += (FVector3f)ApplyTranslation;
+					const FName RootBoneName = Skeleton->GetReferenceSkeleton().GetBoneName(0);
+
+					if(AnimSequence->GetDataModel()->IsValidBoneTrackName(RootBoneName))
+					{
+						TArray<FVector3f> PosKeys;
+						TArray<FQuat4f> RotKeys;
+						TArray<FVector3f> ScaleKeys;
+
+						TArray<FTransform> BoneTransforms;
+						AnimSequence->GetDataModel()->GetBoneTrackTransforms(RootBoneName, BoneTransforms);
+
+						PosKeys.SetNum(BoneTransforms.Num());
+						RotKeys.SetNum(BoneTransforms.Num());
+						ScaleKeys.SetNum(BoneTransforms.Num());
+
+						// Find vector that would translate current root bone location onto origin.
+						FVector FrameTransform = FVector::ZeroVector;
+						if (FrameIndex == INDEX_NONE)
+						{
+							// Use current transform
+							FrameTransform = PreviewSkelComp->GetComponentSpaceTransforms()[0].GetLocation();
+						}
+						else if(BoneTransforms.IsValidIndex(FrameIndex))
+						{
+							// Use transform at frame
+							FrameTransform = BoneTransforms[FrameIndex].GetLocation();
+						}
+
+						FVector ApplyTranslation = -1.f * FrameTransform;
+
+						// Convert into world space
+						const FVector WorldApplyTranslation = PreviewSkelComp->GetComponentTransform().TransformVector(ApplyTranslation);
+						ApplyTranslation = PreviewSkelComp->GetComponentTransform().InverseTransformVector(WorldApplyTranslation);
+
+						for(int32 KeyIndex = 0; KeyIndex < BoneTransforms.Num(); KeyIndex++)
+						{
+							PosKeys[KeyIndex] = FVector3f(BoneTransforms[KeyIndex].GetLocation() + ApplyTranslation);
+							RotKeys[KeyIndex] = FQuat4f(BoneTransforms[KeyIndex].GetRotation());
+							ScaleKeys[KeyIndex] = FVector3f(BoneTransforms[KeyIndex].GetScale3D());
+						}
+
+						IAnimationDataController& Controller = AnimSequence->GetController();
+						Controller.SetBoneTrackKeys(RootBoneName, PosKeys, RotKeys, ScaleKeys);
+					}
 				}
-
-				IAnimationDataController& Controller = AnimSequence->GetController();
-				Controller.SetBoneTrackKeys(AnimationTrack.Name, RawTrack.PosKeys, RawTrack.RotKeys, RawTrack.ScaleKeys);
-
-				AnimSequence->MarkPackageDirty();
 			}
 		}
 	}
diff --git a/Engine/Source/Editor/SequenceRecorder/Private/AnimationRecorder.cpp b/Engine/Source/Editor/SequenceRecorder/Private/AnimationRecorder.cpp
--- a/Engine/Source/Editor/SequenceRecorder/Private/AnimationRecorder.cpp
+++ b/Engine/Source/Editor/SequenceRecorder/Private/AnimationRecorder.cpp
@@ -283,7 +283,7 @@ void FAnimationRecorder::StartRecord(USkeletalMeshComponent* Component, UAnimSeq
 		{
 			// add tracks for the bone existing
 			const FName BoneTreeName = AnimSkeleton->GetReferenceSkeleton().GetBoneName(BoneTreeIndex);
-			Controller.AddBoneTrack(BoneTreeName);			
+			Controller.AddBoneCurve(BoneTreeName);			
 			RawTracks.AddDefaulted();
 		}
 	}
@@ -489,15 +489,14 @@ UAnimSequence* FAnimationRecorder::StopRecord(bool bShowMessage)
 		}
 
 		// Populate bone tracks
-		const TArray<FBoneAnimationTrack>& BoneAnimationTracks = AnimationObject->GetDataModel()->GetBoneAnimationTracks();
-		for (int32 TrackIndex = 0; TrackIndex < BoneAnimationTracks.Num(); ++TrackIndex)
+		TArray<FName> TrackNames;
+		AnimationObject->GetDataModel()->GetBoneTrackNames(TrackNames);
+		for (int32 TrackIndex = 0; TrackIndex < TrackNames.Num(); ++TrackIndex)
 		{
-			const FBoneAnimationTrack& AnimationTrack = BoneAnimationTracks[TrackIndex];
 			const FRawAnimSequenceTrack& RawTrack = RawTracks[TrackIndex];
+			FName BoneName = TrackNames[TrackIndex];
 
-			FName BoneName = AnimationTrack.Name;
-
-			bool bShouldSkipName = ShouldSkipName(AnimationTrack.Name);
+			bool bShouldSkipName = ShouldSkipName(BoneName);
 
 			if (!bShouldSkipName)
 			{
@@ -854,7 +853,9 @@ bool FAnimationRecorder::Record(USkeletalMeshComponent* Component, FTransform co
 			Component->LeaderPoseComponent->GetSkinnedAsset() :
 			Component->GetSkinnedAsset();
 
-		const TArray<FBoneAnimationTrack>& BoneAnimationTracks = AnimationObject->GetDataModel()->GetBoneAnimationTracks();
+		TArray<FName> TrackNames;
+		const IAnimationDataModel* DataModel = AnimationObject->GetDataModel();
+		DataModel->GetBoneTrackNames(TrackNames);
 
 		if (FrameToAdd == 0)
 		{
@@ -862,10 +863,10 @@ bool FAnimationRecorder::Record(USkeletalMeshComponent* Component, FTransform co
 			SkeletonRootIndex = INDEX_NONE;
 			USkeleton* AnimSkeleton = AnimationObject->GetSkeleton();
 
-			for (const FBoneAnimationTrack& AnimationTrack : BoneAnimationTracks)
+			for (const FName& TrackName : TrackNames)
 			{
 				// verify if this bone exists in skeleton
-				const int32 BoneTreeIndex = AnimationTrack.BoneTreeIndex;
+				const int32 BoneTreeIndex = AnimSkeleton->GetReferenceSkeleton().FindBoneIndex(TrackName);
 				if (BoneTreeIndex != INDEX_NONE)
 				{
 					const int32 BoneIndex = AnimSkeleton->GetMeshBoneIndexFromSkeletonBoneIndex(SkinnedAsset, BoneTreeIndex);
@@ -873,7 +874,7 @@ bool FAnimationRecorder::Record(USkeletalMeshComponent* Component, FTransform co
 					const FTransform LocalTransform = SpacesBases[BoneIndex];
 					if (ParentIndex == INDEX_NONE)
 					{
-						if (bRemoveRootTransform && BoneAnimationTracks.Num() > 1)
+						if (bRemoveRootTransform && TrackNames.Num() > 1)
 						{
 							// Store initial root transform.
 							// We remove the initial transform of the root bone and transform root's children
@@ -900,13 +901,13 @@ bool FAnimationRecorder::Record(USkeletalMeshComponent* Component, FTransform co
 		FSerializedAnimation  SerializedAnimation;
 		USkeleton* AnimSkeleton = AnimationObject->GetSkeleton();
 
-		for (int32 TrackIndex = 0; TrackIndex < BoneAnimationTracks.Num(); ++TrackIndex)
+		for (int32 TrackIndex = 0; TrackIndex < TrackNames.Num(); ++TrackIndex)
 		{
-			const FBoneAnimationTrack& AnimationTrack = BoneAnimationTracks[TrackIndex];
+			const FName& TrackName = TrackNames[TrackIndex];
 			FRawAnimSequenceTrack& RawTrack = RawTracks[TrackIndex];
 
 			// verify if this bone exists in skeleton
-			const int32 BoneTreeIndex = AnimationTrack.BoneTreeIndex;
+			const int32 BoneTreeIndex = AnimSkeleton->GetReferenceSkeleton().FindBoneIndex(TrackName);
 			if (BoneTreeIndex != INDEX_NONE)
 			{
 				const int32 BoneIndex = AnimSkeleton->GetMeshBoneIndexFromSkeletonBoneIndex(SkinnedAsset, BoneTreeIndex);
@@ -947,9 +948,9 @@ bool FAnimationRecorder::Record(USkeletalMeshComponent* Component, FTransform co
 					if (FrameToAdd == 0)
 					{
 						const FTransform RefPose = Component->GetSkeletalMeshAsset()->GetRefSkeleton().GetRefBonePose()[BoneIndex];
-						RawTrack.PosKeys.Add((FVector3f)RefPose.GetTranslation());
+						RawTrack.PosKeys.Add(FVector3f(RefPose.GetTranslation()));
 						RawTrack.RotKeys.Add(FQuat4f(RefPose.GetRotation()));
-						RawTrack.ScaleKeys.Add((FVector3f)RefPose.GetScale3D());
+						RawTrack.ScaleKeys.Add(FVector3f(RefPose.GetScale3D()));
 					}
 				}
 			}
diff --git a/Engine/Source/Editor/SequenceRecorder/Private/Sections/MovieScene3DTransformSectionRecorder.cpp b/Engine/Source/Editor/SequenceRecorder/Private/Sections/MovieScene3DTransformSectionRecorder.cpp
--- a/Engine/Source/Editor/SequenceRecorder/Private/Sections/MovieScene3DTransformSectionRecorder.cpp
+++ b/Engine/Source/Editor/SequenceRecorder/Private/Sections/MovieScene3DTransformSectionRecorder.cpp
@@ -143,26 +143,28 @@ void FMovieScene3DTransformSectionRecorder::FinalizeSection(float CurrentTime)
 				int32 RootIndex = INDEX_NONE;
 				USkeleton* AnimSkeleton = AnimSequence->GetSkeleton();
 
-				const TArray<FBoneAnimationTrack>& BoneAnimationTracks = AnimSequence->GetDataModel()->GetBoneAnimationTracks();
-				for (const FBoneAnimationTrack& AnimationTrack : BoneAnimationTracks)
+				TArray<FName> TrackNames;
+				AnimSequence->GetDataModelInterface()->GetBoneTrackNames(TrackNames);
+
+				FName RootBoneName = NAME_None;
+				for (const FName& TrackName : TrackNames)
 				{
-					// verify if this bone exists in skeleton
-					const int32 BoneTreeIndex = AnimationTrack.BoneTreeIndex;
-					if (BoneTreeIndex != INDEX_NONE)
+					// Verify if this bone exists in skeleton
+					const int32 SkeletonTreeIndex = AnimSkeleton->GetReferenceSkeleton().FindBoneIndex(TrackName);
+					if (SkeletonTreeIndex != INDEX_NONE)
 					{
-						const int32 BoneIndex = AnimSkeleton->GetMeshBoneIndexFromSkeletonBoneIndex(SkinnedAsset, BoneTreeIndex);
+						const int32 BoneIndex = AnimSkeleton->GetMeshBoneIndexFromSkeletonBoneIndex(SkinnedAsset, SkeletonTreeIndex);
 						const int32 ParentIndex = SkinnedAsset->GetRefSkeleton().GetParentIndex(BoneIndex);
 						if (ParentIndex == INDEX_NONE)
 						{
-							// found root
+							// We've found the root (root bones do not have a valid parent)
 							RootIndex = BoneIndex;
+							RootBoneName = TrackName;
 							break;
 						}
 					}
 				}
-
-				ensure(BoneAnimationTracks.IsValidIndex(RootIndex));
-				check(RootIndex != INDEX_NONE);
+				check(RootBoneName != NAME_None);
 
 				FFrameRate TickResolution = MovieSceneSection->GetTypedOuter<UMovieScene>()->GetTickResolution();
 				const FFrameNumber StartTime = (RecordingStartTime * TickResolution).FloorToFrame();
@@ -170,38 +172,9 @@ void FMovieScene3DTransformSectionRecorder::FinalizeSection(float CurrentTime)
 				// we may need to offset the transform here if the animation was not recorded on the root component
 				FTransform InvComponentTransform = AnimRecorder->GetComponentTransform().Inverse();
 
-				const FRawAnimSequenceTrack& RawTrack = BoneAnimationTracks[RootIndex].InternalTrackData;
-				const int32 KeyCount = FMath::Max(FMath::Max(RawTrack.PosKeys.Num(), RawTrack.RotKeys.Num()), RawTrack.ScaleKeys.Num());
-				for (int32 KeyIndex = 0; KeyIndex < KeyCount; KeyIndex++)
+				for (int32 KeyIndex = 0; KeyIndex < AnimSequence->GetDataModelInterface()->GetNumberOfKeys(); KeyIndex++)
 				{
-					FTransform Transform;
-					if (RawTrack.PosKeys.IsValidIndex(KeyIndex))
-					{
-						Transform.SetTranslation(FVector(RawTrack.PosKeys[KeyIndex]));
-					}
-					else if (RawTrack.PosKeys.Num() > 0)
-					{
-						Transform.SetTranslation(FVector(RawTrack.PosKeys[0]));
-					}
-
-					if (RawTrack.RotKeys.IsValidIndex(KeyIndex))
-					{
-						Transform.SetRotation(FQuat(RawTrack.RotKeys[KeyIndex]));
-					}
-					else if (RawTrack.RotKeys.Num() > 0)
-					{
-						Transform.SetRotation(FQuat(RawTrack.RotKeys[0]));
-					}
-
-					if (RawTrack.ScaleKeys.IsValidIndex(KeyIndex))
-					{
-						Transform.SetScale3D(FVector(RawTrack.ScaleKeys[KeyIndex]));
-					}
-					else if (RawTrack.ScaleKeys.Num() > 0)
-					{
-						Transform.SetScale3D(FVector(RawTrack.ScaleKeys[0]));
-					}
-
+					const FTransform Transform = AnimSequence->GetDataModelInterface()->GetBoneTrackTransform(RootBoneName, FFrameNumber(KeyIndex));
 					FFrameNumber AnimationFrame = (AnimSequence->GetTimeAtFrame(KeyIndex) * TickResolution).FloorToFrame();
 					BufferedTransforms.Add(InvComponentTransform * Transform, StartTime + AnimationFrame);
 				}
diff --git a/Engine/Source/Editor/SkeletonEditor/Private/EditableSkeleton.cpp b/Engine/Source/Editor/SkeletonEditor/Private/EditableSkeleton.cpp
--- a/Engine/Source/Editor/SkeletonEditor/Private/EditableSkeleton.cpp
+++ b/Engine/Source/Editor/SkeletonEditor/Private/EditableSkeleton.cpp
@@ -317,6 +317,9 @@ void FEditableSkeleton::RenameSmartname(const FName InContainerName, SmartName::
 
 			FAnimationCurveIdentifier CurveId(CurveToRename, ERawCurveTrackTypes::RCT_Float);
 
+			//Make sure skeleton is correct before compression 
+			Skeleton->RenameSmartnameAndModify(InContainerName, InNameUid, InNewName);
+			
 			// Remove curves from animation assets
 			for (FAssetData& Data : AnimationAssets)
 			{
@@ -349,14 +352,11 @@ void FEditableSkeleton::RenameSmartname(const FName InContainerName, SmartName::
 
 			GWarn->BeginSlowTask(LOCTEXT("RebuildingAnimations", "Rebaking/compressing modified animations"), true);
 
-			//Make sure skeleton is correct before compression 
-			Skeleton->RenameSmartnameAndModify(InContainerName, InNameUid, InNewName);
-
 			// Rebake/compress the animations
 			for (UAnimSequence* Seq : SequencesToRecompress)
 			{
 				GWarn->StatusUpdate(1, 2, FText::Format(LOCTEXT("RebuildingAnimationsStatus", "Rebuilding {0}"), FText::FromString(Seq->GetName())));
-				Seq->RequestSyncAnimRecompression();
+				Seq->CacheDerivedDataForCurrentPlatform();
 			}
 
 			GWarn->EndSlowTask();
diff --git a/Engine/Source/Editor/UnrealEd/Private/AnimationEditorUtils.cpp b/Engine/Source/Editor/UnrealEd/Private/AnimationEditorUtils.cpp
--- a/Engine/Source/Editor/UnrealEd/Private/AnimationEditorUtils.cpp
+++ b/Engine/Source/Editor/UnrealEd/Private/AnimationEditorUtils.cpp
@@ -795,8 +795,7 @@ namespace AnimationEditorUtils
 			GWarn->BeginSlowTask(LOCTEXT("AnimCompressing", "Compressing"), true);
 
 			{
-				TSharedPtr<FAnimCompressContext> CompressContext = MakeShareable(new FAnimCompressContext(false, true, AnimSequencePtrs.Num()));
-
+				UE::Anim::Compression::FAnimationCompressionMemorySummaryScope Scope;
 				for (UAnimSequence* AnimSeq : AnimSequencePtrs)
 				{
 					if (OverrideSettings != nullptr)
@@ -806,8 +805,8 @@ namespace AnimationEditorUtils
 
 					// Clear CompressCommandletVersion so we can recompress these animations later.
 					AnimSeq->CompressCommandletVersion = 0;
-					AnimSeq->RequestAnimCompression(FRequestAnimCompressionParams(true, CompressContext));
-					++CompressContext->AnimIndex;
+					AnimSeq->ClearAllCachedCookedPlatformData();
+					AnimSeq->CacheDerivedDataForCurrentPlatform();
 				}
 			}
 
diff --git a/Engine/Source/Editor/UnrealEd/Private/Commandlets/PackageUtilities.cpp b/Engine/Source/Editor/UnrealEd/Private/Commandlets/PackageUtilities.cpp
--- a/Engine/Source/Editor/UnrealEd/Private/Commandlets/PackageUtilities.cpp
+++ b/Engine/Source/Editor/UnrealEd/Private/Commandlets/PackageUtilities.cpp
@@ -1991,7 +1991,7 @@ struct CompressAnimationsFunctor
 			UE_LOG(LogPackageUtilities, Warning, TEXT("%s (%s) Resetting with to default compression settings."), *AnimSeq->GetName(), *AnimSeq->GetFullName());
 			AnimSeq->BoneCompressionSettings = nullptr;
 			AnimSeq->CurveCompressionSettings = nullptr;
-			AnimSeq->RequestAnimCompression(FRequestAnimCompressionParams(false, true, false));
+			AnimSeq->CacheDerivedDataForCurrentPlatform();
 
 			// Automatic compression should have picked a suitable compressor
 			if (!AnimSeq->IsCompressedDataValid())
@@ -1999,7 +1999,7 @@ struct CompressAnimationsFunctor
 				// Update CompressCommandletVersion in that case, and create a proper DDC entry
 				// (with actual compressor)
 				AnimSeq->CompressCommandletVersion = CompressCommandletVersion;
-				AnimSeq->RequestAnimCompression(FRequestAnimCompressionParams(false, false, false));
+				AnimSeq->BeginCacheDerivedDataForCurrentPlatform();
 				bDirtyPackage = true;
 			}
 
diff --git a/Engine/Source/Editor/UnrealEd/Private/Fbx/FbxAnimationExport.cpp b/Engine/Source/Editor/UnrealEd/Private/Fbx/FbxAnimationExport.cpp
--- a/Engine/Source/Editor/UnrealEd/Private/Fbx/FbxAnimationExport.cpp
+++ b/Engine/Source/Editor/UnrealEd/Private/Fbx/FbxAnimationExport.cpp
@@ -131,11 +131,11 @@ void FFbxExporter::ExportAnimSequenceToFbx(const UAnimSequence* AnimSeq, const U
 	{
 		FbxNode* CurrentBoneNode = BoneNodes[BoneIndex];
 		const int32 BoneTreeIndex = Skeleton->GetSkeletonBoneIndexFromMeshBoneIndex(SkelMesh, BoneIndex);
-		const int32 BoneTrackIndex = Skeleton->GetRawAnimationTrackIndex(BoneTreeIndex, AnimSeq);
 		const FName BoneName = Skeleton->GetReferenceSkeleton().GetBoneName(BoneTreeIndex);
+		const IAnimationDataModel* DataModel = AnimSeq->GetDataModel();
 
 		CustomAttributes.Reset();
-		AnimSeq->GetDataModel()->GetAttributesForBone(BoneName, CustomAttributes);
+		DataModel->GetAttributesForBone(BoneName, CustomAttributes);
 
 		TArray<TPair<int32, FbxAnimCurve*>> FloatCustomAttributeIndices;
 		TArray<TPair<int32, FbxAnimCurve*>> IntCustomAttributeIndices;
@@ -143,8 +143,7 @@ void FFbxExporter::ExportAnimSequenceToFbx(const UAnimSequence* AnimSeq, const U
 		// Setup custom attribute properties and curves
 		for (int32 AttributeIndex = 0; AttributeIndex < CustomAttributes.Num(); ++AttributeIndex)
 		{
-			const FAnimatedBoneAttribute* AttributePtr = CustomAttributes[AttributeIndex];
-			if (AttributePtr)
+			if (const FAnimatedBoneAttribute* AttributePtr = CustomAttributes[AttributeIndex])
 			{
 				const FAnimatedBoneAttribute& Attribute = *AttributePtr;
 
@@ -206,7 +205,7 @@ void FFbxExporter::ExportAnimSequenceToFbx(const UAnimSequence* AnimSeq, const U
 		Curves[7] = CurrentBoneNode->LclScaling.GetCurve(InAnimLayer, FBXSDK_CURVENODE_COMPONENT_Y, true);
 		Curves[8] = CurrentBoneNode->LclScaling.GetCurve(InAnimLayer, FBXSDK_CURVENODE_COMPONENT_Z, true);
 
-		if(BoneTrackIndex == INDEX_NONE)
+		if(!DataModel->IsValidBoneTrackName(BoneName))
 		{
 			// If this sequence does not have a track for the current bone, then skip it
 			continue;
@@ -217,24 +216,24 @@ void FFbxExporter::ExportAnimSequenceToFbx(const UAnimSequence* AnimSeq, const U
 			Curve->KeyModifyBegin();
 		}
 
-		auto ExportLambda = [&](double AnimTime, FbxTime ExportTime, bool bLastKey) {
-			FTransform BoneAtom;
-			AnimSeq->GetBoneTransform(BoneAtom, BoneTrackIndex, AnimTime, true);
-			FbxAMatrix FbxMatrix = Converter.ConvertMatrix(BoneAtom.ToMatrixWithScale());
+		auto ExportLambda = [this, DataModel, BoneName, AnimSeq, &FloatCustomAttributeIndices, &IntCustomAttributeIndices, &Curves, &CustomAttributes](double AnimTime, FbxTime ExportTime, bool bLastKey)
+		{
+			const FTransform BoneAtom = DataModel->EvaluateBoneTrackTransform(BoneName, DataModel->GetFrameRate().AsFrameTime(AnimTime), AnimSeq->Interpolation);
+			const FbxAMatrix FbxMatrix = Converter.ConvertMatrix(BoneAtom.ToMatrixWithScale());
 			
-			FbxVector4 Translation = FbxMatrix.GetT();
-			FbxVector4 Rotation = FbxMatrix.GetR();
-			FbxVector4 Scale = FbxMatrix.GetS();
-			FbxVector4 Vectors[3] = { Translation, Rotation, Scale };
+			const FbxVector4 Translation = FbxMatrix.GetT();
+			const FbxVector4 Rotation = FbxMatrix.GetR();
+			const FbxVector4 Scale = FbxMatrix.GetS();
+			const FbxVector4 Vectors[3] = { Translation, Rotation, Scale };
 
 			// Loop over each curve and channel to set correct values
 			for (uint32 CurveIndex = 0; CurveIndex < 3; ++CurveIndex)
 			{
 				for (uint32 ChannelIndex = 0; ChannelIndex < 3; ++ChannelIndex)
 				{
-					uint32 OffsetCurveIndex = (CurveIndex * 3) + ChannelIndex;
+					const uint32 OffsetCurveIndex = (CurveIndex * 3) + ChannelIndex;
 
-					int32 lKeyIndex = Curves[OffsetCurveIndex]->KeyAdd(ExportTime);
+					const int32 lKeyIndex = Curves[OffsetCurveIndex]->KeyAdd(ExportTime);
 					Curves[OffsetCurveIndex]->KeySetValue(lKeyIndex, Vectors[CurveIndex][ChannelIndex]);
 					Curves[OffsetCurveIndex]->KeySetInterpolation(lKeyIndex, bLastKey ? FbxAnimCurveDef::eInterpolationConstant : FbxAnimCurveDef::eInterpolationCubic);
 
@@ -245,31 +244,27 @@ void FFbxExporter::ExportAnimSequenceToFbx(const UAnimSequence* AnimSeq, const U
 				}
 			}
 
-			for (TPair<int32, FbxAnimCurve*>& CurrentAttributeCurve : FloatCustomAttributeIndices)
+			for (const TPair<int32, FbxAnimCurve*>& CurrentAttributeCurve : FloatCustomAttributeIndices)
 			{
-				const FAnimatedBoneAttribute* AttributePtr = CustomAttributes[CurrentAttributeCurve.Key];
-
-				if (AttributePtr)
+				if (const FAnimatedBoneAttribute* AttributePtr = CustomAttributes[CurrentAttributeCurve.Key])
 				{
 					ensure(AttributePtr->Identifier.GetType() == FFloatAnimationAttribute::StaticStruct());
 
-					FFloatAnimationAttribute EvaluatedAttribute = AttributePtr->Curve.Evaluate<FFloatAnimationAttribute>(AnimTime);
-					int32 KeyIndex = CurrentAttributeCurve.Value->KeyAdd(ExportTime);
+					const FFloatAnimationAttribute EvaluatedAttribute = AttributePtr->Curve.Evaluate<FFloatAnimationAttribute>(AnimTime);
+					const int32 KeyIndex = CurrentAttributeCurve.Value->KeyAdd(ExportTime);
 					CurrentAttributeCurve.Value->KeySetValue(KeyIndex, EvaluatedAttribute.Value);
 
 				}
 			}
 
-			for (TPair<int32, FbxAnimCurve*>& CurrentAttributeCurve : IntCustomAttributeIndices)
+			for (const TPair<int32, FbxAnimCurve*>& CurrentAttributeCurve : IntCustomAttributeIndices)
 			{
-				const FAnimatedBoneAttribute* AttributePtr = CustomAttributes[CurrentAttributeCurve.Key];
-
-				if (AttributePtr)
+				if (const FAnimatedBoneAttribute* AttributePtr = CustomAttributes[CurrentAttributeCurve.Key])
 				{
 					ensure(AttributePtr->Identifier.GetType() == FIntegerAnimationAttribute::StaticStruct());
 
-					FIntegerAnimationAttribute EvaluatedAttribute = AttributePtr->Curve.Evaluate<FIntegerAnimationAttribute>(AnimTime);
-					int32 KeyIndex = CurrentAttributeCurve.Value->KeyAdd(ExportTime);
+					const FIntegerAnimationAttribute EvaluatedAttribute = AttributePtr->Curve.Evaluate<FIntegerAnimationAttribute>(AnimTime);
+					const int32 KeyIndex = CurrentAttributeCurve.Value->KeyAdd(ExportTime);
 					CurrentAttributeCurve.Value->KeySetValue(KeyIndex, static_cast<float>(EvaluatedAttribute.Value));
 				}
 			}
diff --git a/Engine/Source/Editor/UnrealEd/Private/PackageAutoSaver.cpp b/Engine/Source/Editor/UnrealEd/Private/PackageAutoSaver.cpp
--- a/Engine/Source/Editor/UnrealEd/Private/PackageAutoSaver.cpp
+++ b/Engine/Source/Editor/UnrealEd/Private/PackageAutoSaver.cpp
@@ -31,7 +31,6 @@
 #include "EditorLevelUtils.h"
 #include "IVREditorModule.h"
 #include "LevelEditorViewport.h"
-#include "Animation/AnimCompressionDerivedDataPublic.h"
 #include "AssetCompilingManager.h"
 
 namespace PackageAutoSaverJson
@@ -507,7 +506,6 @@ bool FPackageAutoSaver::CanAutoSave() const
 	const bool bAreShadersCompiling = GShaderCompilingManager->IsCompiling();
 	const bool bAreAssetsCompiling = FAssetCompilingManager::Get().GetNumRemainingAssets() > 0;
 	const bool bIsVREditorActive = IVREditorModule::Get().IsVREditorEnabled();	// @todo vreditor: Eventually we should support this while in VR (modal VR progress, with sufficient early warning)
-	const bool bAreAnimationsCompressing = GAsyncCompressedAnimationsTracker ? GAsyncCompressedAnimationsTracker->GetNumRemainingJobs() > 0 : false;
 	const bool bIsInterchangeActive = UInterchangeManager::GetInterchangeManager().IsInterchangeActive();
 
 	bool bIsSequencerPlaying = false;
@@ -533,7 +531,6 @@ bool FPackageAutoSaver::CanAutoSave() const
 		&& bHasGameOrProjectLoaded
 		&& !bAreShadersCompiling
 		&& !bAreAssetsCompiling
-		&& !bAreAnimationsCompressing
 		&& !bIsVREditorActive
 		&& !bIsSequencerPlaying
 		&& !bIsInterchangeActive
diff --git a/Engine/Source/Editor/UnrealEd/Private/SkeletalMeshEdit.cpp b/Engine/Source/Editor/UnrealEd/Private/SkeletalMeshEdit.cpp
--- a/Engine/Source/Editor/UnrealEd/Private/SkeletalMeshEdit.cpp
+++ b/Engine/Source/Editor/UnrealEd/Private/SkeletalMeshEdit.cpp
@@ -1283,7 +1283,6 @@ namespace AnimationTransformDebug
 	// Data sturctutre to debug bone transform of animation issues
 	struct FAnimationTransformDebugData
 	{
-		int32 TrackIndex;
 		int32 BoneIndex;
 		FName BoneName;
 		TArray<FTransform>	RecalculatedLocalTransform;
@@ -1295,12 +1294,11 @@ namespace AnimationTransformDebug
 		TArray<FTransform>	SourceParentGlobalTransform;
 
 		FAnimationTransformDebugData()
-			: TrackIndex(INDEX_NONE), BoneIndex(INDEX_NONE), BoneName(NAME_None)
+			: BoneIndex(INDEX_NONE), BoneName(NAME_None)
 		{}
 
-		void SetTrackData(int32 InTrackIndex, int32 InBoneIndex, FName InBoneName)
+		void SetTrackData( int32 InBoneIndex, FName InBoneName)
 		{
-			TrackIndex = InTrackIndex;
 			BoneIndex = InBoneIndex;
 			BoneName = InBoneName;
 		}
@@ -2326,11 +2324,9 @@ void UnFbx::FFbxImporter::ImportBoneTracks(USkeleton* Skeleton, FAnimCurveImport
 				if (BoneTreeIndex != INDEX_NONE)
 				{
 					//add new track
-					Controller.AddBoneTrack(BoneName);
+					Controller.AddBoneCurve(BoneName);
 					Controller.SetBoneTrackKeys(BoneName, RawTrack.PosKeys, RawTrack.RotKeys, RawTrack.ScaleKeys);
-					const int32 TrackIndex = DestSeq->GetDataModel()->GetBoneTrackIndexByName(BoneName);
-
-					NewDebugData.SetTrackData(TrackIndex, BoneTreeIndex, BoneName);
+					NewDebugData.SetTrackData(BoneTreeIndex, BoneName);
 
 					// add mapping to skeleton bone track
 					TransformDebugData.Add(NewDebugData);
diff --git a/Engine/Source/Runtime/Core/Private/Math/UnrealMath.cpp b/Engine/Source/Runtime/Core/Private/Math/UnrealMath.cpp
--- a/Engine/Source/Runtime/Core/Private/Math/UnrealMath.cpp
+++ b/Engine/Source/Runtime/Core/Private/Math/UnrealMath.cpp
@@ -469,6 +469,7 @@ FQuat4f FRotator3f::Quaternion() const
 	{
 		logOrEnsureNanError(TEXT("Invalid input %s to FRotator::Quaternion - generated NaN output: %s"), *ToString(), *RotationQuat.ToString());
 		RotationQuat = FQuat4f::Identity;
+		FDebug::DumpStackTraceToLog(ELogVerbosity::Warning);
 	}
 #endif
 
@@ -538,6 +539,7 @@ FQuat4d FRotator3d::Quaternion() const
 	{
 		logOrEnsureNanError(TEXT("Invalid input %s to FRotator::Quaternion - generated NaN output: %s"), *ToString(), *RotationQuat.ToString());
 		RotationQuat = FQuat4d::Identity;
+		FDebug::DumpStackTraceToLog(ELogVerbosity::Warning);
 	}
 #endif
 
diff --git a/Engine/Source/Runtime/Core/Private/Misc/FrameRate.cpp b/Engine/Source/Runtime/Core/Private/Misc/FrameRate.cpp
--- a/Engine/Source/Runtime/Core/Private/Misc/FrameRate.cpp
+++ b/Engine/Source/Runtime/Core/Private/Misc/FrameRate.cpp
@@ -355,4 +355,20 @@ bool TryParseString(FFrameRate& OutFrameRate, const TCHAR* InString)
 	return false;
 }
 
+FArchive& operator<<(FArchive& Ar, FFrameRate& FrameRate)
+{
+	Ar << FrameRate.Numerator;
+	Ar << FrameRate.Denominator;
+
+	return Ar;
+}
+
+bool FFrameRate::Serialize(FArchive& Ar)
+{
+	Ar << Numerator;
+	Ar << Denominator;
+	
+	return true;
+}
+
 #undef LOCTEXT_NAMESPACE
\ No newline at end of file
diff --git a/Engine/Source/Runtime/Core/Public/Misc/FrameRate.h b/Engine/Source/Runtime/Core/Public/Misc/FrameRate.h
--- a/Engine/Source/Runtime/Core/Public/Misc/FrameRate.h
+++ b/Engine/Source/Runtime/Core/Public/Misc/FrameRate.h
@@ -207,6 +207,10 @@ struct FFrameRate
 	{
 		return Rate.AsFrameTime(TimeInSeconds);
 	}
+	
+	friend CORE_API FArchive& operator<<(FArchive& Ar, FFrameRate& FrameRate);
+
+	CORE_API bool Serialize(FArchive& Ar);
 };
 
 inline double FFrameRate::AsInterval() const
diff --git a/Engine/Source/Runtime/Core/Public/UObject/UE5MainStreamObjectVersion.h b/Engine/Source/Runtime/Core/Public/UObject/UE5MainStreamObjectVersion.h
--- a/Engine/Source/Runtime/Core/Public/UObject/UE5MainStreamObjectVersion.h
+++ b/Engine/Source/Runtime/Core/Public/UObject/UE5MainStreamObjectVersion.h
@@ -319,6 +319,9 @@ struct CORE_API FUE5MainStreamObjectVersion
 		// Store serialized graph function in the function data
 		RigVMSaveSerializedGraphInGraphFunctionData,
 
+		// Animation Sequence now stores its frame-rate on a per-platform basis
+		PerPlatformAnimSequenceTargetFrameRate,
+
 		// -----<new versions can be added above this line>-------------------------------------------------
 		VersionPlusOne,
 		LatestVersion = VersionPlusOne - 1
diff --git a/Engine/Source/Runtime/Core/Public/UObject/UnrealNames.inl b/Engine/Source/Runtime/Core/Public/UObject/UnrealNames.inl
--- a/Engine/Source/Runtime/Core/Public/UObject/UnrealNames.inl
+++ b/Engine/Source/Runtime/Core/Public/UObject/UnrealNames.inl
@@ -257,6 +257,9 @@ REGISTER_NAME(408,InterpCurveVector)
 REGISTER_NAME(409,InterpCurveTwoVectors)
 REGISTER_NAME(410,InterpCurveQuat)
 
+
+REGISTER_NAME(411,FrameRate)
+
 REGISTER_NAME(450,AI)
 REGISTER_NAME(451,NavMesh)
 
diff --git a/Engine/Source/Runtime/Engine/Classes/Animation/AnimBoneCompressionSettings.h b/Engine/Source/Runtime/Engine/Classes/Animation/AnimBoneCompressionSettings.h
--- a/Engine/Source/Runtime/Engine/Classes/Animation/AnimBoneCompressionSettings.h
+++ b/Engine/Source/Runtime/Engine/Classes/Animation/AnimBoneCompressionSettings.h
@@ -7,9 +7,12 @@
 
 class UAnimBoneCompressionCodec;
 class UAnimSequenceBase;
+
+#if WITH_EDITORONLY_DATA
 struct FCompressibleAnimData;
 struct FCompressibleAnimDataResult;
 namespace UE::Anim::Compression { struct FAnimDDCKeyArgs; }
+#endif // WITH_EDITORONLY_DATA
 
 /*
  * This object is used to wrap a bone compression codec. It allows a clean integration in the editor by avoiding the need
diff --git a/Engine/Source/Runtime/Engine/Classes/Animation/AnimCompress.h b/Engine/Source/Runtime/Engine/Classes/Animation/AnimCompress.h
--- a/Engine/Source/Runtime/Engine/Classes/Animation/AnimCompress.h
+++ b/Engine/Source/Runtime/Engine/Classes/Animation/AnimCompress.h
@@ -191,7 +191,7 @@ class ENGINE_API FCompressionMemorySummary
 public:
 	FCompressionMemorySummary(bool bInEnabled);
 
-	void GatherPreCompressionStats(const FString& Name, int32 RawSize, int32 PreviousCompressionSize, int32 ProgressNumerator, int32 ProgressDenominator);
+	void GatherPreCompressionStats(int32 RawSize, int32 PreviousCompressionSize);
 
 	void GatherPostCompressionStats(const FCompressedAnimSequence& CompressedData, const TArray<FBoneData>& BoneData, const FName AnimFName, double CompressionTime, bool bInPerformedCompression);
 
@@ -224,27 +224,25 @@ private:
 //////////////////////////////////////////////////////////////////////////
 // FAnimCompressContext - Context information / storage for use during
 // animation compression
+
+struct UE_DEPRECATED(5.2, "FAnimCompressContext has been deprecated") FAnimCompressContext;
 struct ENGINE_API FAnimCompressContext
 {
 private:
 	FCompressionMemorySummary	CompressionSummary;
 
-	void GatherPreCompressionStats(const FString& Name, int32 RawSize, int32 PreviousCompressionSize);
-
-	void GatherPostCompressionStats(const FCompressedAnimSequence& CompressedData, const TArray<FBoneData>& BoneData, const FName AnimFName, double CompressionTime, bool bInPerformedCompression);
-
+	void GatherPreCompressionStats(const FString& Name, int32 RawSize, int32 PreviousCompressionSize) {}
 
+	void GatherPostCompressionStats(const FCompressedAnimSequence& CompressedData, const TArray<FBoneData>& BoneData, const FName AnimFName, double CompressionTime, bool bInPerformedCompression) {}
 public:
 	uint32						AnimIndex;
 	uint32						MaxAnimations;
-	bool						bAllowAlternateCompressor;
 	bool						bOutput;
 
-	FAnimCompressContext(bool bInAllowAlternateCompressor, bool bInOutput, uint32 InMaxAnimations = 1)
+	FAnimCompressContext(bool bInOutput, uint32 InMaxAnimations = 1)
 		: CompressionSummary(bInOutput)
 		, AnimIndex(0)
 		, MaxAnimations(InMaxAnimations)
-		, bAllowAlternateCompressor(bInAllowAlternateCompressor)
 		, bOutput(bInOutput)
 	{}
 
@@ -253,7 +251,6 @@ public:
 		: CompressionSummary(false)
 		, AnimIndex(Rhs.AnimIndex)
 		, MaxAnimations(Rhs.MaxAnimations)
-		, bAllowAlternateCompressor(Rhs.bAllowAlternateCompressor)
 		, bOutput(Rhs.bOutput)
 	{}
 
@@ -262,6 +259,54 @@ public:
 	friend class UAnimSequence;
 };
 
+#if WITH_EDITOR
+namespace UE
+{
+	namespace Anim
+	{
+		namespace Compression
+		{		
+			// This is a version string that mimics the old versioning scheme. If you
+			// want to bump this version, generate a new guid using VS->Tools->Create GUID and
+			// return it here. Ex.
+			static FString AnimationCompressionVersionString = TEXT("FB8D0FF1ED0C4848951601EB6BB028F6");
+			
+			struct ENGINE_API FAnimationCompressionMemorySummaryScope
+			{
+				FAnimationCompressionMemorySummaryScope()
+				{
+					bool bExpected = false;
+					check(ScopeExists.compare_exchange_strong(bExpected, true));
+					CompressionSummary = MakeUnique<FCompressionMemorySummary>(true);
+				}
+
+				~FAnimationCompressionMemorySummaryScope()
+				{
+					bool bExpected = true;
+					check(ScopeExists.compare_exchange_strong(bExpected, false));
+					CompressionSummary.Reset();
+				}
+
+				static bool ShouldStoreCompressionResults()
+				{
+					return ScopeExists.load();
+				}
+	
+				static FCompressionMemorySummary& CompressionResultSummary()
+				{
+					check(ScopeExists.load());
+					return *CompressionSummary.Get();
+				}
+
+				static std::atomic<bool> ScopeExists;
+				static TUniquePtr<FCompressionMemorySummary> CompressionSummary;
+			};
+		}
+	}
+}
+#endif // WITH_EDITOR
+
+
 UCLASS(abstract, hidecategories=Object, MinimalAPI, EditInlineNew)
 class UAnimCompress : public UAnimBoneCompressionCodec
 {
diff --git a/Engine/Source/Runtime/Engine/Classes/Animation/AnimData/AnimDataModel.h b/Engine/Source/Runtime/Engine/Classes/Animation/AnimData/AnimDataModel.h
--- a/Engine/Source/Runtime/Engine/Classes/Animation/AnimData/AnimDataModel.h
+++ b/Engine/Source/Runtime/Engine/Classes/Animation/AnimData/AnimDataModel.h
@@ -29,6 +29,8 @@ public:
 	virtual int32 GetNumberOfFrames() const override;
 	virtual int32 GetNumberOfKeys() const override;
 	virtual FFrameRate GetFrameRate() const override;
+	
+	PRAGMA_DISABLE_DEPRECATION_WARNINGS	
 	virtual const TArray<FBoneAnimationTrack>& GetBoneAnimationTracks() const override;
 	virtual const FBoneAnimationTrack& GetBoneTrackByIndex(int32 TrackIndex) const override;
 	virtual const FBoneAnimationTrack& GetBoneTrackByName(FName TrackName) const override;
@@ -37,9 +39,16 @@ public:
 	virtual int32 GetBoneTrackIndex(const FBoneAnimationTrack& Track) const override;
 	virtual int32 GetBoneTrackIndexByName(FName TrackName) const override;
 	virtual bool IsValidBoneTrackIndex(int32 TrackIndex) const override;
+	PRAGMA_ENABLE_DEPRECATION_WARNINGS
+	
 	virtual int32 GetNumBoneTracks() const override;
 	virtual void GetBoneTrackNames(TArray<FName>& OutNames) const override;
-
+	virtual FTransform GetBoneTrackTransform(FName TrackName, const FFrameNumber& FrameNumber) const override;
+	virtual void GetBoneTrackTransforms(FName TrackName, const TArray<FFrameNumber>& FrameNumbers, TArray<FTransform>& OutTransforms) const override;	
+	virtual void GetBoneTrackTransforms(FName TrackName, TArray<FTransform>& OutTransforms) const override;
+	virtual void GetBoneTracksTransform(const TArray<FName>& TrackNames, const FFrameNumber& FrameNumber, TArray<FTransform>& OutTransforms) const override;	
+	virtual FTransform EvaluateBoneTrackTransform(FName TrackName, const FFrameTime& FrameTime, const EAnimInterpolationType& Interpolation) const override;
+	virtual bool IsValidBoneTrackName(const FName& TrackName) const override;	
 	virtual const FAnimCurveBase* FindCurve(const FAnimationCurveIdentifier& CurveIdentifier) const override;
 	virtual const FFloatCurve* FindFloatCurve(const FAnimationCurveIdentifier& CurveIdentifier) const override;
 	virtual const FTransformCurve* FindTransformCurve(const FAnimationCurveIdentifier& CurveIdentifier) const override;
@@ -67,10 +76,13 @@ public:
 #endif
 	virtual TScriptInterface<IAnimationDataController> GetController() override;
 	virtual bool HasBeenPopulated() const override { return bPopulated; }
+	virtual void IterateBoneKeys(const FName& BoneName, TFunction<bool(const FVector3f& Pos, const FQuat4f&, const FVector3f, const FFrameNumber&)> IterationFunction) const override;
 protected:
 	virtual IAnimationDataModel::FModelNotifier& GetNotifier() override;
 	virtual FAnimDataModelModifiedDynamicEvent& GetModifiedDynamicEvent() override { return ModifiedEventDynamic; }
 	virtual void OnNotify(const EAnimDataModelNotifyType& NotifyType, const FAnimDataModelNotifPayload& Payload) override {}
+	virtual void LockEvaluationAndModification() const override final {}
+	virtual void UnlockEvaluationAndModification() const override final {}
 	/** End IAnimationDataModel overrides */
 
 private:
diff --git a/Engine/Source/Runtime/Engine/Classes/Animation/AnimData/AnimDataNotifications.h b/Engine/Source/Runtime/Engine/Classes/Animation/AnimData/AnimDataNotifications.h
--- a/Engine/Source/Runtime/Engine/Classes/Animation/AnimData/AnimDataNotifications.h
+++ b/Engine/Source/Runtime/Engine/Classes/Animation/AnimData/AnimDataNotifications.h
@@ -111,8 +111,18 @@ USTRUCT(BlueprintType)
 struct FAnimationTrackAddedPayload : public FAnimationTrackPayload
 {
 	GENERATED_BODY()
+	
+	
+	PRAGMA_DISABLE_DEPRECATION_WARNINGS
+	FAnimationTrackAddedPayload() = default;
+	FAnimationTrackAddedPayload(const FAnimationTrackAddedPayload&) = default;
+	FAnimationTrackAddedPayload(FAnimationTrackAddedPayload&&) = default;
+	FAnimationTrackAddedPayload& operator=(const FAnimationTrackAddedPayload&) = default;
+	FAnimationTrackAddedPayload& operator=(FAnimationTrackAddedPayload&&) = default;
+	PRAGMA_ENABLE_DEPRECATION_WARNINGS
 
 	/** Index of the track (bone) which was added */
+	UE_DEPRECATED(5.2, "FAnimationTrackAddedPayload::TrackIndex has been deprecated")
 	UPROPERTY(BlueprintReadOnly, Category = Payload)
 	int32 TrackIndex = INDEX_NONE;
 };
diff --git a/Engine/Source/Runtime/Engine/Classes/Animation/AnimData/IAnimationDataController.h b/Engine/Source/Runtime/Engine/Classes/Animation/AnimData/IAnimationDataController.h
--- a/Engine/Source/Runtime/Engine/Classes/Animation/AnimData/IAnimationDataController.h
+++ b/Engine/Source/Runtime/Engine/Classes/Animation/AnimData/IAnimationDataController.h
@@ -187,8 +187,12 @@ public:
 	*
 	* @return	The index at which the bone track was added, INDEX_NONE if adding it failed
 	*/
+	UE_DEPRECATED(5.2, "AddBoneTrack returning index is deprecated use AddBoneCurve returning bool instead.")
+	UFUNCTION(BlueprintCallable, Category = AnimationData,  meta=(DeprecatedFunction, DeprecationMessage="AddBoneTrack returning index is deprecated use AddBoneCurve returning bool instead."))
+	virtual int32 AddBoneTrack(FName BoneName, bool bShouldTransact = true) { AddBoneCurve(BoneName, bShouldTransact); return INDEX_NONE; }
+
 	UFUNCTION(BlueprintCallable, Category = AnimationData)
-	virtual int32 AddBoneTrack(FName BoneName, bool bShouldTransact = true) = 0;
+	virtual bool AddBoneCurve(FName BoneName, bool bShouldTransact = true) = 0;
 
 	/**
 	* Inserts a new bone animation track for the provided name, at the provided index. Broadcasts a EAnimDataModelNotifyType::TrackAdded notify if successful.
@@ -809,11 +813,11 @@ protected:
 	}
 
 	template<typename ActionClass, class... ActionArguments>
-	void ConditionalAction(bool bCondition, ActionArguments... Arguments)
+	void ConditionalAction(bool bCondition, ActionArguments&&... Arguments)
 	{
 		if (UE::FChangeTransactor::CanTransactChanges() && bCondition) \
 		{ 
-			ChangeTransactor.AddTransactionChange<ActionClass>(Arguments...); \
+			ChangeTransactor.AddTransactionChange<ActionClass>(Forward<ActionArguments>(Arguments)...); \
 		}
 	}
 protected:
