Crash when starting a take with a Metahuman source - No Movie Scene found for SequencerDataModel
#fix make sure we never try and sample time-code attributes on an animation sequence which has not been populated (DataModel is created but not initialised)

diff --git a/Engine/Source/Editor/AnimationBlueprintLibrary/Private/AnimationBlueprintLibrary.cpp b/Engine/Source/Editor/AnimationBlueprintLibrary/Private/AnimationBlueprintLibrary.cpp
--- a/Engine/Source/Editor/AnimationBlueprintLibrary/Private/AnimationBlueprintLibrary.cpp
+++ b/Engine/Source/Editor/AnimationBlueprintLibrary/Private/AnimationBlueprintLibrary.cpp
@@ -2039,7 +2039,7 @@ bool UAnimationBlueprintLibrary::EvaluateRootBoneTimecodeAttributesAtTime(const
 	}
 
 	const IAnimationDataModel* AnimDataModel = AnimationSequenceBase->GetDataModel();
-	if (!AnimDataModel)
+	if (!AnimDataModel || !AnimDataModel->HasBeenPopulated())
 	{
 		return false;
 	}
