Determine the location of a Socket in an Animation Sequence
#feature added functionality to FAnimPose API to retrieve animated Socket transforms
#misc ptr safety check

diff --git a/Engine/Source/Editor/AnimationBlueprintLibrary/Private/AnimPose.cpp b/Engine/Source/Editor/AnimationBlueprintLibrary/Private/AnimPose.cpp
--- a/Engine/Source/Editor/AnimationBlueprintLibrary/Private/AnimPose.cpp
+++ b/Engine/Source/Editor/AnimationBlueprintLibrary/Private/AnimPose.cpp
@@ -31,6 +31,7 @@
 #include "Misc/MemStack.h"
 #include "PreviewScene.h"
 #include "ReferenceSkeleton.h"
+#include "Engine/SkeletalMeshSocket.h"
 #include "Templates/Casts.h"
 #include "Trace/Detail/Channel.h"
 #include "UObject/Object.h"
@@ -173,6 +174,29 @@ void FAnimPose::SetPose(const FAnimationPoseData& PoseData)
 			CurveNames.Add(InElement.Name);
 			CurveValues.Add(InElement.Value);
 		});
+
+		TArray<USkeletalMeshSocket*> Sockets;
+		const USkeleton* Skeleton = ContextBoneContainer.GetSkeletonAsset();
+		const USkeletalMesh* SkeletalMesh = ContextBoneContainer.GetSkeletalMeshAsset();
+		if (SkeletalMesh)
+		{
+			Sockets = SkeletalMesh->GetActiveSocketList();
+		}
+		else if (Skeleton)
+		{
+			Sockets = Skeleton->Sockets;
+		}
+
+		for (const USkeletalMeshSocket* Socket : Sockets)
+		{
+			const int32 PoseBoneIndex = ContextBoneContainer.GetPoseBoneIndexForBoneName(Socket->BoneName);
+			if (PoseBoneIndex != INDEX_NONE)
+			{
+				SocketNames.Add(Socket->SocketName);
+				SocketParentBoneNames.Add(Socket->BoneName);
+				SocketTransforms.Add(Socket->GetSocketLocalTransform());
+			}
+		}		
 	}
 	else
 	{
@@ -384,6 +408,38 @@ FTransform UAnimPoseExtensions::GetRefPoseRelativeTransform(const FAnimPose& Pos
 	return FTransform::Identity;
 }
 
+void UAnimPoseExtensions::GetSocketNames(const FAnimPose& Pose, TArray<FName>& Sockets)
+{
+	if (Pose.IsValid())
+	{
+		Sockets = Pose.SocketNames;
+	}
+}
+
+FTransform UAnimPoseExtensions::GetSocketPose(const FAnimPose& Pose, FName SocketName, EAnimPoseSpaces Space)
+{
+	if (Pose.IsValid())
+	{
+		const int32 SocketIndex = Pose.SocketNames.IndexOfByKey(SocketName);			
+		if (SocketIndex != INDEX_NONE)
+		{
+			const int32 BoneIndex = Pose.BoneNames.IndexOfByKey(Pose.SocketParentBoneNames[SocketIndex]);
+			const FTransform BoneTransform = Space == EAnimPoseSpaces::Local ? Pose.LocalSpacePoses[BoneIndex] : Pose.WorldSpacePoses[BoneIndex];
+			return Pose.SocketTransforms[SocketIndex] * BoneTransform;
+		}
+		else
+		{
+			UE_LOG(LogAnimationPoseScripting, Warning, TEXT("No socket with name %s was found"), *SocketName.ToString());
+		}
+	}
+	else
+	{		
+		UE_LOG(LogAnimationPoseScripting, Error, TEXT("Provided Pose is not valid"));	
+	}
+		
+	return FTransform::Identity;
+}
+
 void UAnimPoseExtensions::EvaluateAnimationBlueprintWithInputPose(const FAnimPose& Pose, USkeletalMesh* TargetSkeletalMesh, UAnimBlueprint* AnimationBlueprint, FAnimPose& OutPose)
 {
 	if (Pose.IsValid())
@@ -515,8 +571,11 @@ float UAnimPoseExtensions::GetCurveWeight(const FAnimPose& Pose, const FName& Cu
 
 void UAnimPoseExtensions::GetAnimPoseAtFrame(const UAnimSequenceBase* AnimationSequenceBase, int32 FrameIndex, FAnimPoseEvaluationOptions EvaluationOptions, FAnimPose& Pose)
 {
-	const double Time = AnimationSequenceBase->GetDataModel()->GetFrameRate().AsSeconds(FrameIndex);
-	GetAnimPoseAtTime(AnimationSequenceBase, Time, EvaluationOptions, Pose);
+	if (AnimationSequenceBase)
+	{
+		const double Time = AnimationSequenceBase->GetDataModel()->GetFrameRate().AsSeconds(FrameIndex);
+		GetAnimPoseAtTime(AnimationSequenceBase, Time, EvaluationOptions, Pose);
+	}
 }
 
 void UAnimPoseExtensions::GetAnimPoseAtTime(const UAnimSequenceBase* AnimationSequenceBase, double Time, FAnimPoseEvaluationOptions EvaluationOptions, FAnimPose& Pose)
diff --git a/Engine/Source/Editor/AnimationBlueprintLibrary/Public/AnimPose.h b/Engine/Source/Editor/AnimationBlueprintLibrary/Public/AnimPose.h
--- a/Engine/Source/Editor/AnimationBlueprintLibrary/Public/AnimPose.h
+++ b/Engine/Source/Editor/AnimationBlueprintLibrary/Public/AnimPose.h
@@ -134,6 +134,15 @@ protected:
 
 	UPROPERTY()
 	TArray<float> CurveValues;
+	
+	UPROPERTY()
+	TArray<FName> SocketNames;
+
+	UPROPERTY()
+	TArray<FName> SocketParentBoneNames;
+	
+	UPROPERTY()
+	TArray<FTransform> SocketTransforms;
 
 	friend class UAnimPoseExtensions;
 };
@@ -238,6 +247,29 @@ public:
 	UFUNCTION(BlueprintPure, meta = (ScriptMethod), Category = "Animation|Pose")
     static FTransform GetRefPoseRelativeTransform(UPARAM(ref) const FAnimPose& Pose, FName FromBoneName, FName ToBoneName, EAnimPoseSpaces Space = EAnimPoseSpaces::Local);
 
+	/**
+	* Returns an array of socket names contained by the pose
+	*
+	* @param	Pose	Anim Pose to retrieve the names from
+	* @param	Sockets	Array to be populated with the socket names
+	*
+	*/
+	UFUNCTION(BlueprintPure, meta = (ScriptMethod), Category = "Animation|Pose")
+	static void GetSocketNames(UPARAM(ref) const FAnimPose& Pose, TArray<FName>& Sockets);
+
+	/**
+	* Retrieves the transform for the provided socket name from a pose
+	*
+	* @param	Pose		Anim Pose to retrieve the transform from
+	* @param	SocketName	Name of the socket to retrieve
+	* @param	Space		Space in which the transform should be retrieved
+	*
+	* @return	Transform in requested space for bone if found, otherwise return identity transform
+	*/
+	UFUNCTION(BlueprintPure, meta = (ScriptMethod), Category = "Animation|Pose")
+	static FTransform GetSocketPose(UPARAM(ref) const FAnimPose& Pose, FName SocketName, EAnimPoseSpaces Space = EAnimPoseSpaces::Local);
+	
+
 	/**
 	* Evaluates an Animation Sequence Base to generate a valid Anim Pose instance
 	*
