Reinstate some anim blueprint library methods which now use the Anim Pose api.

diff --git a/Engine/Source/Editor/AnimationBlueprintLibrary/Private/AnimationBlueprintLibrary.cpp b/Engine/Source/Editor/AnimationBlueprintLibrary/Private/AnimationBlueprintLibrary.cpp
--- a/Engine/Source/Editor/AnimationBlueprintLibrary/Private/AnimationBlueprintLibrary.cpp
+++ b/Engine/Source/Editor/AnimationBlueprintLibrary/Private/AnimationBlueprintLibrary.cpp
@@ -62,6 +62,7 @@
 #include "UObject/ObjectPtr.h"
 #include "UObject/UnrealNames.h"
 #include "Animation/AnimData/IAnimationDataModel.h"
+#include "AnimPose.h"
 
 #define LOCTEXT_NAMESPACE "AnimationBlueprintLibrary"
 
@@ -70,6 +71,57 @@ IMPLEMENT_MODULE(IModuleInterface, AnimationBlueprintLibrary);
 
 DEFINE_LOG_CATEGORY_STATIC(LogAnimationBlueprintLibrary, Verbose, All);
 
+static void GetBonePosesForTimeInternal(const UAnimSequenceBase* AnimationSequenceBase, TArray<FName> BoneNames, float Time, bool bExtractRootMotion, TArray<FTransform>& Poses, const USkeletalMesh* PreviewMesh = nullptr)
+{
+	Poses.Empty(BoneNames.Num());
+	if (AnimationSequenceBase && AnimationSequenceBase->GetSkeleton())
+	{
+		Poses.AddDefaulted(BoneNames.Num());
+
+		if (BoneNames.Num())
+		{
+			TArray<FName> TrackNames;
+			AnimationSequenceBase->GetDataModel()->GetBoneTrackNames(TrackNames);
+			
+			for (int32 BoneNameIndex = 0; BoneNameIndex < BoneNames.Num(); ++BoneNameIndex)
+			{
+				const FName& BoneName = BoneNames[BoneNameIndex];
+				if (TrackNames.Contains(BoneName))
+				{
+					FAnimPoseEvaluationOptions EvaluationOptions = FAnimPoseEvaluationOptions();
+					FAnimPose AnimPose;
+    
+					UAnimPoseExtensions::GetAnimPoseAtTime(AnimationSequenceBase, Time, EvaluationOptions, AnimPose);
+					Poses[BoneNameIndex] = UAnimPoseExtensions::GetBonePose(AnimPose, BoneName, EAnimPoseSpaces::Local);	
+				}
+				else
+				{
+					// otherwise, get ref pose if exists
+					const FReferenceSkeleton& RefSkeleton = (PreviewMesh)? PreviewMesh->GetRefSkeleton() : AnimationSequenceBase->GetSkeleton()->GetReferenceSkeleton();
+					const int32 BoneIndex = RefSkeleton.FindBoneIndex(BoneName);
+					if (BoneIndex != INDEX_NONE)
+					{
+						Poses[BoneNameIndex] = RefSkeleton.GetRefBonePose()[BoneIndex];
+					}
+					else
+					{
+						UE_LOG(LogAnimationBlueprintLibrary, Warning, TEXT("Invalid bone name %s for Animation Sequence %s supplied for GetBonePosesForTime"), *BoneName.ToString(), *AnimationSequenceBase->GetName());
+						Poses[BoneNameIndex] = FTransform::Identity;
+					}
+				}			
+			}
+		}
+		else
+		{
+			UE_LOG(LogAnimationBlueprintLibrary, Error, TEXT("Invalid or no bone names specified to retrieve poses given Animation Sequence %s in GetBonePosesForTime"), *AnimationSequenceBase->GetName());
+		}	
+	}
+	else
+	{
+		UE_LOG(LogAnimationBlueprintLibrary, Warning, TEXT("Invalid Animation Sequence supplied for GetBonePosesForTime"));
+	}
+}
+
 void UAnimationBlueprintLibrary::GetNumFrames(const UAnimSequenceBase* AnimationSequenceBase, int32& NumFrames)
 {
 	NumFrames = 0;
@@ -2411,6 +2463,58 @@ void UAnimationBlueprintLibrary::AddNodeAssetOverride(UAnimBlueprint* AnimBluepr
 	}
 }
 
+void UAnimationBlueprintLibrary::GetBonePoseForTime(const UAnimSequenceBase* AnimationSequenceBase, FName BoneName, float Time, bool bExtractRootMotion, FTransform& Pose)
+{
+	Pose.SetIdentity();
+	if (AnimationSequenceBase)
+	{
+		TArray<FName> BoneNameArray;
+		TArray<FTransform> PoseArray;
+		BoneNameArray.Add(BoneName);
+		GetBonePosesForTimeInternal(AnimationSequenceBase, BoneNameArray, Time, bExtractRootMotion, PoseArray);
+		Pose = PoseArray[0];
+	}
+	else
+	{
+		UE_LOG(LogAnimationBlueprintLibrary, Warning, TEXT("Invalid Animation Sequence supplied for GetBonePoseForTime"));
+	}
+}
+
+void UAnimationBlueprintLibrary::GetBonePoseForFrame(const UAnimSequenceBase* AnimationSequenceBase, FName BoneName, int32 Frame, bool bExtractRootMotion, FTransform& Pose)
+{
+	Pose.SetIdentity();
+	if (AnimationSequenceBase)
+	{
+		TArray<FName> BoneNameArray;
+		TArray<FTransform> PoseArray;
+		BoneNameArray.Add(BoneName);
+		GetBonePosesForTimeInternal(AnimationSequenceBase, BoneNameArray, GetTimeAtFrameInternal(AnimationSequenceBase, Frame), bExtractRootMotion, PoseArray);
+		Pose = PoseArray[0];
+	}
+	else
+	{
+		UE_LOG(LogAnimationBlueprintLibrary, Warning, TEXT("Invalid Animation Sequence supplied for GetBonePoseForFrame"));
+	}
+}
+
+void UAnimationBlueprintLibrary::GetBonePosesForTime(const UAnimSequenceBase* AnimationSequenceBase, TArray<FName> BoneNames, float Time, bool bExtractRootMotion, TArray<FTransform>& Poses, const USkeletalMesh* PreviewMesh /*= nullptr*/)
+{
+	GetBonePosesForTimeInternal(AnimationSequenceBase, BoneNames, Time, bExtractRootMotion, Poses, PreviewMesh);
+}
+
+void UAnimationBlueprintLibrary::GetBonePosesForFrame(const UAnimSequenceBase* AnimationSequenceBase, TArray<FName> BoneNames, int32 Frame, bool bExtractRootMotion, TArray<FTransform>& Poses, const USkeletalMesh* PreviewMesh /*= nullptr*/)
+{
+	Poses.Empty(BoneNames.Num());
+	if (AnimationSequenceBase)
+	{
+		GetBonePosesForTimeInternal(AnimationSequenceBase, BoneNames, GetTimeAtFrameInternal(AnimationSequenceBase, Frame), bExtractRootMotion, Poses, PreviewMesh);
+	}
+	else
+	{
+		UE_LOG(LogAnimationBlueprintLibrary, Warning, TEXT("Invalid Animation Sequence supplied for GetBonePosesForFrame"));
+	}
+}
+
 template void UAnimationBlueprintLibrary::AddCurveKeysInternal<float, FFloatCurve, ERawCurveTrackTypes::RCT_Float>(UAnimSequence* AnimationSequence, FName CurveName, const TArray<float>& Times, const TArray<float>& KeyData);
 template void UAnimationBlueprintLibrary::AddCurveKeysInternal<FVector, FVectorCurve, ERawCurveTrackTypes::RCT_Vector>(UAnimSequence* AnimationSequence, FName CurveName, const TArray<float>& Times, const TArray<FVector>& KeyData);
 template void UAnimationBlueprintLibrary::AddCurveKeysInternal<FTransform, FTransformCurve, ERawCurveTrackTypes::RCT_Transform>(UAnimSequence* AnimationSequence, FName CurveName, const TArray<float>& Times, const TArray<FTransform>& KeyData);
diff --git a/Engine/Source/Editor/AnimationBlueprintLibrary/Public/AnimationBlueprintLibrary.h b/Engine/Source/Editor/AnimationBlueprintLibrary/Public/AnimationBlueprintLibrary.h
--- a/Engine/Source/Editor/AnimationBlueprintLibrary/Public/AnimationBlueprintLibrary.h
+++ b/Engine/Source/Editor/AnimationBlueprintLibrary/Public/AnimationBlueprintLibrary.h
@@ -459,22 +459,22 @@ public:
 	/** Retrieves Bone Pose data for the given Bone Name at the specified Time from the given Animation Sequence */
 	UE_DEPRECATED(5.2, "GetBonePosesForTime is deprecated, use AnimPose or AnimationDataModel interface directly")
 	UFUNCTION(BlueprintPure, Category = "AnimationBlueprintLibrary|Pose")
-	static void GetBonePoseForTime(const UAnimSequenceBase* AnimationSequenceBase, FName BoneName, float Time, bool bExtractRootMotion, FTransform& Pose) {}
+	static void GetBonePoseForTime(const UAnimSequenceBase* AnimationSequenceBase, FName BoneName, float Time, bool bExtractRootMotion, FTransform& Pose);
 
 	/** Retrieves Bone Pose data for the given Bone Name at the specified Frame from the given Animation Sequence */
 	UE_DEPRECATED(5.2, "GetBonePosesForTime is deprecated, use AnimPose or AnimationDataModel interface directly")
 	UFUNCTION(BlueprintPure, Category = "AnimationBlueprintLibrary|Pose")
-	static void GetBonePoseForFrame(const UAnimSequenceBase* AnimationSequenceBase, FName BoneName, int32 Frame, bool bExtractRootMotion, FTransform& Pose) {}
+	static void GetBonePoseForFrame(const UAnimSequenceBase* AnimationSequenceBase, FName BoneName, int32 Frame, bool bExtractRootMotion, FTransform& Pose);
 
 	/** Retrieves Bone Pose data for the given Bone Names at the specified Time from the given Animation Sequence */
 	UE_DEPRECATED(5.2, "GetBonePosesForTime is deprecated, use AnimPose or AnimationDataModel interface directly")
 	UFUNCTION(BlueprintPure, Category = "AnimationBlueprintLibrary|Pose")
-	static void GetBonePosesForTime(const UAnimSequenceBase* AnimationSequenceBase, TArray<FName> BoneNames, float Time, bool bExtractRootMotion, TArray<FTransform>& Poses, const USkeletalMesh* PreviewMesh = nullptr) {}
+	static void GetBonePosesForTime(const UAnimSequenceBase* AnimationSequenceBase, TArray<FName> BoneNames, float Time, bool bExtractRootMotion, TArray<FTransform>& Poses, const USkeletalMesh* PreviewMesh = nullptr);
 
 	/** Retrieves Bone Pose data for the given Bone Names at the specified Frame from the given Animation Sequence */
 	UE_DEPRECATED(5.2, "GetBonePosesForTime is deprecated, use AnimPose or AnimationDataModel interface directly")
 	UFUNCTION(BlueprintPure, Category = "AnimationBlueprintLibrary|Pose")
-	static void GetBonePosesForFrame(const UAnimSequenceBase* AnimationSequenceBase, TArray<FName> BoneNames, int32 Frame, bool bExtractRootMotion, TArray<FTransform>& Poses, const USkeletalMesh* PreviewMesh = nullptr) {}
+	static void GetBonePosesForFrame(const UAnimSequenceBase* AnimationSequenceBase, TArray<FName> BoneNames, int32 Frame, bool bExtractRootMotion, TArray<FTransform>& Poses, const USkeletalMesh* PreviewMesh = nullptr);
 
 	// Virtual bones
 
