Added UFunction missing deprecation/message keywords to deprecated Animation APIs

diff --git a/Engine/Source/Editor/AnimationBlueprintLibrary/Public/AnimationBlueprintLibrary.h b/Engine/Source/Editor/AnimationBlueprintLibrary/Public/AnimationBlueprintLibrary.h
--- a/Engine/Source/Editor/AnimationBlueprintLibrary/Public/AnimationBlueprintLibrary.h
+++ b/Engine/Source/Editor/AnimationBlueprintLibrary/Public/AnimationBlueprintLibrary.h
@@ -88,22 +88,22 @@ public:
 
 	/** Retrieves the Raw Translation Animation Data for the given Animation Track Name and Animation Sequence */
 	UE_DEPRECATED(5.2, "GetRawTrackPositionData has been deprecated, use AnimationModel interface instead")
-	UFUNCTION(BlueprintPure, Category = "AnimationBlueprintLibrary|RawTrackData")
+	UFUNCTION(BlueprintPure, Category = "AnimationBlueprintLibrary|RawTrackData", meta=(DeprecatedFunction, DeprecationMessage="GetRawTrackPositionData has been deprecated, use AnimationModel interface instead"))
 	static void GetRawTrackPositionData(const UAnimSequenceBase* AnimationSequenceBase, const FName TrackName, TArray<FVector>& PositionData) {}
 
 	/** Retrieves the Raw Rotation Animation Data for the given Animation Track Name and Animation Sequence */
 	UE_DEPRECATED(5.2, "GetRawTrackRotationData has been deprecated, use AnimationModel interface instead")
-	UFUNCTION(BlueprintPure, Category = "AnimationBlueprintLibrary|RawTrackData")
+	UFUNCTION(BlueprintPure, Category = "AnimationBlueprintLibrary|RawTrackData", meta=(DeprecatedFunction, DeprecationMessage="GetRawTrackRotationData has been deprecated, use AnimationModel interface instead"))
 	static void GetRawTrackRotationData(const UAnimSequenceBase* AnimationSequenceBase, const FName TrackName, TArray<FQuat>& RotationData ) {}
 
 	/** Retrieves the Raw Scale Animation Data for the given Animation Track Name and Animation Sequence */	
 	UE_DEPRECATED(5.2, "GetRawTrackScaleData has been deprecated, use AnimationModel interface instead")
-	UFUNCTION(BlueprintPure, Category = "AnimationBlueprintLibrary|RawTrackData")
+	UFUNCTION(BlueprintPure, Category = "AnimationBlueprintLibrary|RawTrackData", meta=(DeprecatedFunction, DeprecationMessage="GetRawTrackScaleData has been deprecated, use AnimationModel interface instead"))
 	static void GetRawTrackScaleData(const UAnimSequenceBase* AnimationSequenceBase, const FName TrackName, TArray<FVector>& ScaleData) {}
 
 	/** Retrieves the Raw Animation Data for the given Animation Track Name and Animation Sequence */
-	UE_DEPRECATED(5.2, "GetRawTrackScaleData has been deprecated, use AnimationModel interface instead")
-	UFUNCTION(BlueprintPure, Category = "AnimationBlueprintLibrary|RawTrackData")
+	UE_DEPRECATED(5.2, "GetRawTrackData has been deprecated, use AnimationModel interface instead")
+	UFUNCTION(BlueprintPure, Category = "AnimationBlueprintLibrary|RawTrackData", meta=(DeprecatedFunction, DeprecationMessage="GetRawTrackData has been deprecated, use AnimationModel interface instead"))
 	static void GetRawTrackData(const UAnimSequenceBase* AnimationSequenceBase, const FName TrackName, TArray<FVector>& PositionKeys,TArray<FQuat>& RotationKeys, TArray<FVector>& ScalingKeys) {}
 
 	/** Checks whether or not the given Animation Track Name is contained within the Animation Sequence */
@@ -467,22 +467,22 @@ public:
 
 	/** Retrieves Bone Pose data for the given Bone Name at the specified Time from the given Animation Sequence */
 	UE_DEPRECATED(5.2, "GetBonePosesForTime is deprecated, use AnimPose or AnimationDataModel interface directly")
-	UFUNCTION(BlueprintPure, Category = "AnimationBlueprintLibrary|Pose")
+	UFUNCTION(BlueprintPure, Category = "AnimationBlueprintLibrary|Pose", meta=(DeprecatedFunction, DeprecationMessage="GetBonePoseForTime has been deprecated, use AnimationModel interface instead"))
 	static void GetBonePoseForTime(const UAnimSequenceBase* AnimationSequenceBase, FName BoneName, float Time, bool bExtractRootMotion, FTransform& Pose);
 
 	/** Retrieves Bone Pose data for the given Bone Name at the specified Frame from the given Animation Sequence */
 	UE_DEPRECATED(5.2, "GetBonePosesForTime is deprecated, use AnimPose or AnimationDataModel interface directly")
-	UFUNCTION(BlueprintPure, Category = "AnimationBlueprintLibrary|Pose")
+	UFUNCTION(BlueprintPure, Category = "AnimationBlueprintLibrary|Pose", meta=(DeprecatedFunction, DeprecationMessage="GetBonePosesForTime has been deprecated, use AnimationModel interface instead"))
 	static void GetBonePoseForFrame(const UAnimSequenceBase* AnimationSequenceBase, FName BoneName, int32 Frame, bool bExtractRootMotion, FTransform& Pose);
 
 	/** Retrieves Bone Pose data for the given Bone Names at the specified Time from the given Animation Sequence */
 	UE_DEPRECATED(5.2, "GetBonePosesForTime is deprecated, use AnimPose or AnimationDataModel interface directly")
-	UFUNCTION(BlueprintPure, Category = "AnimationBlueprintLibrary|Pose")
+	UFUNCTION(BlueprintPure, Category = "AnimationBlueprintLibrary|Pose", meta=(DeprecatedFunction, DeprecationMessage="GetBonePosesForTime has been deprecated, use AnimationModel interface instead"))
 	static void GetBonePosesForTime(const UAnimSequenceBase* AnimationSequenceBase, TArray<FName> BoneNames, float Time, bool bExtractRootMotion, TArray<FTransform>& Poses, const USkeletalMesh* PreviewMesh = nullptr);
 
 	/** Retrieves Bone Pose data for the given Bone Names at the specified Frame from the given Animation Sequence */
 	UE_DEPRECATED(5.2, "GetBonePosesForTime is deprecated, use AnimPose or AnimationDataModel interface directly")
-	UFUNCTION(BlueprintPure, Category = "AnimationBlueprintLibrary|Pose")
+	UFUNCTION(BlueprintPure, Category = "AnimationBlueprintLibrary|Pose", meta=(DeprecatedFunction, DeprecationMessage="GetBonePosesForTime has been deprecated, use AnimationModel interface instead"))
 	static void GetBonePosesForFrame(const UAnimSequenceBase* AnimationSequenceBase, TArray<FName> BoneNames, int32 Frame, bool bExtractRootMotion, TArray<FTransform>& Poses, const USkeletalMesh* PreviewMesh = nullptr);
 
 	// Virtual bones
diff --git a/Engine/Source/Runtime/Engine/Classes/Animation/AnimData/IAnimationDataModel.h b/Engine/Source/Runtime/Engine/Classes/Animation/AnimData/IAnimationDataModel.h
--- a/Engine/Source/Runtime/Engine/Classes/Animation/AnimData/IAnimationDataModel.h
+++ b/Engine/Source/Runtime/Engine/Classes/Animation/AnimData/IAnimationDataModel.h
@@ -156,7 +156,7 @@ public:
 	* @return	Array containing all bone animation tracks 
 	*/
 	UE_DEPRECATED(5.2, "GetBoneAnimationTracks has been deprecated")
-	UFUNCTION(BlueprintCallable, Category = AnimationDataModel)
+	UFUNCTION(BlueprintCallable, Category = AnimationDataModel, meta=(DeprecatedFunction, DeprecationMessage="GetBoneAnimationTracks has been deprecated"))
 	virtual const TArray<FBoneAnimationTrack>& GetBoneAnimationTracks() const = 0;
 
 	virtual FTransform EvaluateBoneTrackTransform(FName TrackName, const FFrameTime& FrameTime, const EAnimInterpolationType& Interpolation) const = 0;
@@ -169,14 +169,14 @@ public:
 	* @return	Bone animation track for the provided index
 	*/
 	UE_DEPRECATED(5.2, "GetBoneTrackByIndex has been deprecated")
-	UFUNCTION(BlueprintCallable, Category = AnimationDataModel)
+	UFUNCTION(BlueprintCallable, Category = AnimationDataModel, meta=(DeprecatedFunction, DeprecationMessage="GetBoneTrackByIndex has been deprecated"))
 	virtual const FBoneAnimationTrack& GetBoneTrackByIndex(int32 TrackIndex) const = 0;
 
 	/**
 	* @return	Bone animation track for the provided (bone) name
 	*/
 	UE_DEPRECATED(5.2, "GetBoneTrackByName has been deprecated")
-	UFUNCTION(BlueprintCallable, Category = AnimationDataModel)
+	UFUNCTION(BlueprintCallable, Category = AnimationDataModel, meta=(DeprecatedFunction, DeprecationMessage="GetBoneTrackByName has been deprecated"))
 	virtual const FBoneAnimationTrack& GetBoneTrackByName(FName TrackName) const = 0;
 
 	/**
@@ -195,21 +195,21 @@ public:
 	* @return	Internal track index for the provided bone animation track if found, otherwise returns INDEX_NONE 
 	*/
 	UE_DEPRECATED(5.2, "FindBoneTrackByIndex has been deprecated")
-	UFUNCTION(BlueprintCallable, Category = AnimationDataModel)
+	UFUNCTION(BlueprintCallable, Category = AnimationDataModel, meta=(DeprecatedFunction, DeprecationMessage="FindBoneTrackByIndex has been deprecated"))
 	virtual int32 GetBoneTrackIndex(const FBoneAnimationTrack& Track) const = 0;
 
 	/**
 	* @return	Internal track index for the provided (bone) name if found, otherwise returns INDEX_NONE 
 	*/
 	UE_DEPRECATED(5.2, "GetBoneTrackIndexByName has been deprecated")
-	UFUNCTION(BlueprintCallable, Category = AnimationDataModel)
+	UFUNCTION(BlueprintCallable, Category = AnimationDataModel, meta=(DeprecatedFunction, DeprecationMessage="GetBoneTrackIndexByName has been deprecated"))
 	virtual int32 GetBoneTrackIndexByName(FName TrackName) const = 0;
 
 	/**
 	* @return	Whether or not the provided track index is valid 
 	*/
 	UE_DEPRECATED(5.2, "IsValidBoneTrackIndex has been deprecated")
-	UFUNCTION(BlueprintCallable, Category = AnimationDataModel)
+	UFUNCTION(BlueprintCallable, Category = AnimationDataModel, meta=(DeprecatedFunction, DeprecationMessage="IsValidBoneTrackIndex has been deprecated"))
 	virtual bool IsValidBoneTrackIndex(int32 TrackIndex) const = 0;
 
 	UFUNCTION(BlueprintCallable, Category = AnimationDataModel)
diff --git a/Engine/Source/Runtime/Engine/Classes/Animation/AnimInstance.h b/Engine/Source/Runtime/Engine/Classes/Animation/AnimInstance.h
--- a/Engine/Source/Runtime/Engine/Classes/Animation/AnimInstance.h
+++ b/Engine/Source/Runtime/Engine/Classes/Animation/AnimInstance.h
@@ -1127,7 +1127,7 @@ public:
 	ENGINE_API void ClearMorphTargets();
 
 	UE_DEPRECATED(5.0, "Please use UKismetAnimationLibrary::CalculateDirection instead")
-	UFUNCTION(BlueprintCallable, Category="Animation", meta=(BlueprintThreadSafe))
+	UFUNCTION(BlueprintCallable, Category="Animation", meta=(BlueprintThreadSafe, DeprecatedFunction, DeprecationMessage="Please use UKismetAnimationLibrary::CalculateDirection instead"))
 	ENGINE_API float CalculateDirection(const FVector& Velocity, const FRotator& BaseRotation) const;
 
 	//--- AI communication start ---//
