Sequencer: Check for valid skeleton in EvaluateRootBoneTimecodeAttributesAtTime

diff --git a/Engine/Source/Editor/AnimationBlueprintLibrary/Private/AnimationBlueprintLibrary.cpp b/Engine/Source/Editor/AnimationBlueprintLibrary/Private/AnimationBlueprintLibrary.cpp
--- a/Engine/Source/Editor/AnimationBlueprintLibrary/Private/AnimationBlueprintLibrary.cpp
+++ b/Engine/Source/Editor/AnimationBlueprintLibrary/Private/AnimationBlueprintLibrary.cpp
@@ -2009,7 +2009,7 @@ bool UAnimationBlueprintLibrary::IsValidTimeInternal(const UAnimSequenceBase* An
 
 bool UAnimationBlueprintLibrary::EvaluateRootBoneTimecodeAttributesAtTime(const UAnimSequenceBase* AnimationSequenceBase, const float EvalTime, FQualifiedFrameTime& OutQualifiedFrameTime)
 {
-	if (!AnimationSequenceBase)
+	if (!AnimationSequenceBase || !AnimationSequenceBase->GetSkeleton())
 	{
 		return false;
 	}
@@ -2195,7 +2195,7 @@ bool UAnimationBlueprintLibrary::EvaluateRootBoneTimecodeAttributesAtTime(const
 
 bool UAnimationBlueprintLibrary::EvaluateRootBoneTimecodeSubframeAttributeAtTime(const UAnimSequenceBase* AnimationSequenceBase, const float EvalTime, float& OutSubframe)
 {
-	if (!AnimationSequenceBase)
+	if (!AnimationSequenceBase || !AnimationSequenceBase->GetSkeleton())
 	{
 		return false;
 	}
