PR #10124: AnimationBlueprintLibrary Curve Functions

diff --git a/Engine/Source/Editor/AnimationBlueprintLibrary/Private/AnimationBlueprintLibrary.cpp b/Engine/Source/Editor/AnimationBlueprintLibrary/Private/AnimationBlueprintLibrary.cpp
--- a/Engine/Source/Editor/AnimationBlueprintLibrary/Private/AnimationBlueprintLibrary.cpp
+++ b/Engine/Source/Editor/AnimationBlueprintLibrary/Private/AnimationBlueprintLibrary.cpp
@@ -47,13 +47,11 @@
 #include "Misc/FrameRate.h"
 #include "Misc/FrameTime.h"
 #include "Misc/Guid.h"
-#include "Misc/MemStack.h"
 #include "Misc/QualifiedFrameTime.h"
 #include "Misc/Timecode.h"
 #include "Modules/ModuleInterface.h"
 #include "Modules/ModuleManager.h"
 #include "ReferenceSkeleton.h"
-#include "Serialization/StructuredArchiveAdapters.h"
 #include "Templates/Casts.h"
 #include "Templates/UnrealTemplate.h"
 #include "Trace/Detail/Channel.h"
@@ -178,10 +176,10 @@ void UAnimationBlueprintLibrary::GetMontageSlotNames(const UAnimMontage* Animati
 	}
 }
 
-void UAnimationBlueprintLibrary::GetAnimationCurveNames(const UAnimSequence* AnimationSequence, ERawCurveTrackTypes CurveType, TArray<FName>& CurveNames)
+void UAnimationBlueprintLibrary::GetAnimationCurveNames(const UAnimSequenceBase* AnimationSequenceBase, ERawCurveTrackTypes CurveType, TArray<FName>& CurveNames)
 {
 	CurveNames.Empty();
-	if (AnimationSequence)
+	if (AnimationSequenceBase)
 	{
 		auto GetCurveName = [](const auto& Curve) -> FName
 		{
@@ -192,13 +190,13 @@ void UAnimationBlueprintLibrary::GetAnimationCurveNames(const UAnimSequence* Ani
 		{
 			case ERawCurveTrackTypes::RCT_Float:
 			{
-				Algo::Transform(AnimationSequence->GetDataModel()->GetFloatCurves(), CurveNames, GetCurveName);
+				Algo::Transform(AnimationSequenceBase->GetDataModel()->GetFloatCurves(), CurveNames, GetCurveName);
 				break;
 			}
 
 			case ERawCurveTrackTypes::RCT_Transform:
 			{
-				Algo::Transform(AnimationSequence->GetDataModel()->GetTransformCurves(), CurveNames, GetCurveName);
+				Algo::Transform(AnimationSequenceBase->GetDataModel()->GetTransformCurves(), CurveNames, GetCurveName);
 				break;
 			}
 
@@ -1257,9 +1255,9 @@ void UAnimationBlueprintLibrary::GetAnimationNotifyEventsForTrack(const UAnimSeq
 	}
 }
 
-void UAnimationBlueprintLibrary::AddCurve(UAnimSequence* AnimationSequence, FName CurveName, ERawCurveTrackTypes CurveType /*= RCT_Float*/, bool bMetaDataCurve /*= false*/)
+void UAnimationBlueprintLibrary::AddCurve(UAnimSequenceBase* AnimationSequenceBase, FName CurveName, ERawCurveTrackTypes CurveType /*= RCT_Float*/, bool bMetaDataCurve /*= false*/)
 {
-	if (AnimationSequence)
+	if (AnimationSequenceBase)
 	{
 		static const ESmartNameContainerType ContainerForCurveType[(int32)ERawCurveTrackTypes::RCT_MAX] = { ESmartNameContainerType::SNCT_CurveMapping, ESmartNameContainerType::SNCT_CurveMapping, ESmartNameContainerType::SNCT_TrackCurveMapping };
 		const ESmartNameContainerType CurveContainer = ContainerForCurveType[(int32)CurveType];
@@ -1270,17 +1268,17 @@ void UAnimationBlueprintLibrary::AddCurve(UAnimSequence* AnimationSequence, FNam
 		// Only Float metadata curves are valid
 		const bool bValidMetaData = !bMetaDataCurve || (bMetaDataCurve && CurveType == ERawCurveTrackTypes::RCT_Float);
 		// Transform curves can only be added if the curve name exists as a bone on the skeleton
-		const bool bValidTransformCurveData = CurveType != ERawCurveTrackTypes::RCT_Transform || (AnimationSequence->GetSkeleton() && DoesBoneNameExistInternal(AnimationSequence->GetSkeleton(), CurveName));
+		const bool bValidTransformCurveData = CurveType != ERawCurveTrackTypes::RCT_Transform || (AnimationSequenceBase->GetSkeleton() && DoesBoneNameExistInternal(AnimationSequenceBase->GetSkeleton(), CurveName));
 
 		if (bValidMetaData && bValidTransformCurveData )
 		{
 			// Add or retrieve the smartname
-			const bool bCurveAdded = AddCurveInternal(AnimationSequence, CurveName, CurveFlags, CurveType);
+			const bool bCurveAdded = AddCurveInternal(AnimationSequenceBase, CurveName, CurveFlags, CurveType);
 
 			if (!bCurveAdded)
 			{
 				// Curve already existed
-				UE_LOG(LogAnimationBlueprintLibrary, Warning, TEXT("Curve %s already exists on the Skeleton %s."), *CurveName.ToString(), *AnimationSequence->GetSkeleton()->GetName());
+				UE_LOG(LogAnimationBlueprintLibrary, Warning, TEXT("Curve %s already exists on the Skeleton %s."), *CurveName.ToString(), *AnimationSequenceBase->GetSkeleton()->GetName());
 			}
 		}
 		else
@@ -1292,7 +1290,7 @@ void UAnimationBlueprintLibrary::AddCurve(UAnimSequence* AnimationSequence, FNam
 			
 			if (!bValidTransformCurveData)
 			{
-				UE_LOG(LogAnimationBlueprintLibrary, Warning, TEXT("Invalid Transform Curve name, the supplied name %s does not exist on the Skeleton %s."), *CurveName.ToString(), AnimationSequence->GetSkeleton() ? *AnimationSequence->GetSkeleton()->GetName() : TEXT("Invalid Skeleton"));
+				UE_LOG(LogAnimationBlueprintLibrary, Warning, TEXT("Invalid Transform Curve name, the supplied name %s does not exist on the Skeleton %s."), *CurveName.ToString(), AnimationSequenceBase->GetSkeleton() ? *AnimationSequenceBase->GetSkeleton()->GetName() : TEXT("Invalid Skeleton"));
 			}
 		}
 	}
@@ -1302,14 +1300,14 @@ void UAnimationBlueprintLibrary::AddCurve(UAnimSequence* AnimationSequence, FNam
 	}	
 }
 
-void UAnimationBlueprintLibrary::RemoveCurve(UAnimSequence* AnimationSequence, FName CurveName, bool bRemoveNameFromSkeleton /*= false*/)
+void UAnimationBlueprintLibrary::RemoveCurve(UAnimSequenceBase* AnimationSequenceBase, FName CurveName, bool bRemoveNameFromSkeleton /*= false*/)
 {
-	if (AnimationSequence)
+	if (AnimationSequenceBase)
 	{
-		const ERawCurveTrackTypes CurveType = RetrieveCurveTypeForCurve(AnimationSequence, CurveName);
+		const ERawCurveTrackTypes CurveType = RetrieveCurveTypeForCurve(AnimationSequenceBase, CurveName);
 		if (CurveType != ERawCurveTrackTypes::RCT_MAX)
 		{
-			const bool bCurveRemoved = RemoveCurveInternal(AnimationSequence, CurveName, CurveType);
+			const bool bCurveRemoved = RemoveCurveInternal(AnimationSequenceBase, CurveName, CurveType);
 		}
 		else
 		{
@@ -1322,11 +1320,11 @@ void UAnimationBlueprintLibrary::RemoveCurve(UAnimSequence* AnimationSequence, F
 	}
 }
 
-void UAnimationBlueprintLibrary::RemoveAllCurveData(UAnimSequence* AnimationSequence)
+void UAnimationBlueprintLibrary::RemoveAllCurveData(UAnimSequenceBase* AnimationSequenceBase)
 {
-	if (AnimationSequence)
+	if (AnimationSequenceBase)
 	{
-		IAnimationDataController& Controller = AnimationSequence->GetController();
+		IAnimationDataController& Controller = AnimationSequenceBase->GetController();
 
 		Controller.RemoveAllCurvesOfType(ERawCurveTrackTypes::RCT_Float);
 		Controller.RemoveAllCurvesOfType(ERawCurveTrackTypes::RCT_Transform);
@@ -1337,9 +1335,9 @@ void UAnimationBlueprintLibrary::RemoveAllCurveData(UAnimSequence* AnimationSequ
 	}
 }
 
-void UAnimationBlueprintLibrary::AddTransformationCurveKey(UAnimSequence* AnimationSequence, FName CurveName, const float Time, const FTransform& Transform)
+void UAnimationBlueprintLibrary::AddTransformationCurveKey(UAnimSequenceBase* AnimationSequenceBase, FName CurveName, const float Time, const FTransform& Transform)
 {
-	if (AnimationSequence)
+	if (AnimationSequenceBase)
 	{
 		TArray<float> TimeArray;
 		TArray<FTransform> TransformArray;
@@ -1347,7 +1345,7 @@ void UAnimationBlueprintLibrary::AddTransformationCurveKey(UAnimSequence* Animat
 		TimeArray.Add(Time);
 		TransformArray.Add(Transform);
 
-		AddCurveKeysInternal<FTransform, FTransformCurve, ERawCurveTrackTypes::RCT_Transform>(AnimationSequence, CurveName, TimeArray, TransformArray);
+		AddCurveKeysInternal<FTransform, FTransformCurve, ERawCurveTrackTypes::RCT_Transform>(AnimationSequenceBase, CurveName, TimeArray, TransformArray);
 	}
 	else
 	{
@@ -1356,13 +1354,13 @@ void UAnimationBlueprintLibrary::AddTransformationCurveKey(UAnimSequence* Animat
 
 }
 
-void UAnimationBlueprintLibrary::AddTransformationCurveKeys(UAnimSequence* AnimationSequence, FName CurveName, const TArray<float>& Times, const TArray<FTransform>& Transforms)
+void UAnimationBlueprintLibrary::AddTransformationCurveKeys(UAnimSequenceBase* AnimationSequenceBase, FName CurveName, const TArray<float>& Times, const TArray<FTransform>& Transforms)
 {
-	if (AnimationSequence)
+	if (AnimationSequenceBase)
 	{
 		if (Times.Num() == Transforms.Num())
 		{
-			AddCurveKeysInternal<FTransform, FTransformCurve, ERawCurveTrackTypes::RCT_Transform>(AnimationSequence, CurveName, Times, Transforms);
+			AddCurveKeysInternal<FTransform, FTransformCurve, ERawCurveTrackTypes::RCT_Transform>(AnimationSequenceBase, CurveName, Times, Transforms);
 		}
 		else
 		{
@@ -1376,9 +1374,9 @@ void UAnimationBlueprintLibrary::AddTransformationCurveKeys(UAnimSequence* Anima
 }
 
 
-void UAnimationBlueprintLibrary::AddFloatCurveKey(UAnimSequence* AnimationSequence, FName CurveName, const float Time, const float Value)
+void UAnimationBlueprintLibrary::AddFloatCurveKey(UAnimSequenceBase* AnimationSequenceBase, FName CurveName, const float Time, const float Value)
 {
-	if (AnimationSequence)
+	if (AnimationSequenceBase)
 	{
 		TArray<float> TimeArray;
 		TArray<float> ValueArray;
@@ -1386,7 +1384,7 @@ void UAnimationBlueprintLibrary::AddFloatCurveKey(UAnimSequence* AnimationSequen
 		TimeArray.Add(Time);
 		ValueArray.Add(Value);
 
-		AddCurveKeysInternal<float, FFloatCurve, ERawCurveTrackTypes::RCT_Float>(AnimationSequence, CurveName, TimeArray, ValueArray);
+		AddCurveKeysInternal<float, FFloatCurve, ERawCurveTrackTypes::RCT_Float>(AnimationSequenceBase, CurveName, TimeArray, ValueArray);
 	}
 	else
 	{
@@ -1395,13 +1393,13 @@ void UAnimationBlueprintLibrary::AddFloatCurveKey(UAnimSequence* AnimationSequen
 
 }
 
-void UAnimationBlueprintLibrary::AddFloatCurveKeys(UAnimSequence* AnimationSequence, FName CurveName, const TArray<float>& Times, const TArray<float>& Values)
+void UAnimationBlueprintLibrary::AddFloatCurveKeys(UAnimSequenceBase* AnimationSequenceBase, FName CurveName, const TArray<float>& Times, const TArray<float>& Values)
 {
-	if (AnimationSequence)
+	if (AnimationSequenceBase)
 	{
 		if (Times.Num() == Values.Num())
 		{
-			AddCurveKeysInternal<float, FFloatCurve, ERawCurveTrackTypes::RCT_Float>(AnimationSequence, CurveName, Times, Values);
+			AddCurveKeysInternal<float, FFloatCurve, ERawCurveTrackTypes::RCT_Float>(AnimationSequenceBase, CurveName, Times, Values);
 		}
 		else
 		{
@@ -1416,9 +1414,9 @@ void UAnimationBlueprintLibrary::AddFloatCurveKeys(UAnimSequence* AnimationSeque
 	
 }
 
-void UAnimationBlueprintLibrary::AddVectorCurveKey(UAnimSequence* AnimationSequence, FName CurveName, const float Time, const FVector Vector)
+void UAnimationBlueprintLibrary::AddVectorCurveKey(UAnimSequenceBase* AnimationSequenceBase, FName CurveName, const float Time, const FVector Vector)
 {
-	if (AnimationSequence)
+	if (AnimationSequenceBase)
 	{
 		TArray<float> TimeArray;
 		TArray<FVector> VectorArray;
@@ -1426,7 +1424,7 @@ void UAnimationBlueprintLibrary::AddVectorCurveKey(UAnimSequence* AnimationSeque
 		TimeArray.Add(Time);
 		VectorArray.Add(Vector);
 
-		AddCurveKeysInternal<FVector, FVectorCurve, ERawCurveTrackTypes::RCT_Vector>(AnimationSequence, CurveName, TimeArray, VectorArray);
+		AddCurveKeysInternal<FVector, FVectorCurve, ERawCurveTrackTypes::RCT_Vector>(AnimationSequenceBase, CurveName, TimeArray, VectorArray);
 	}
 	else
 	{
@@ -1435,13 +1433,13 @@ void UAnimationBlueprintLibrary::AddVectorCurveKey(UAnimSequence* AnimationSeque
 
 }
 
-void UAnimationBlueprintLibrary::AddVectorCurveKeys(UAnimSequence* AnimationSequence, FName CurveName, const TArray<float>& Times, const TArray<FVector>& Vectors)
+void UAnimationBlueprintLibrary::AddVectorCurveKeys(UAnimSequenceBase* AnimationSequenceBase, FName CurveName, const TArray<float>& Times, const TArray<FVector>& Vectors)
 {
-	if (AnimationSequence)
+	if (AnimationSequenceBase)
 	{
 		if (Times.Num() == Vectors.Num())
 		{
-			AddCurveKeysInternal<FVector, FVectorCurve, ERawCurveTrackTypes::RCT_Vector>(AnimationSequence, CurveName, Times, Vectors);
+			AddCurveKeysInternal<FVector, FVectorCurve, ERawCurveTrackTypes::RCT_Vector>(AnimationSequenceBase, CurveName, Times, Vectors);
 		}
 		else
 		{
@@ -1480,33 +1478,33 @@ static void SetControllerCurveKeys(IAnimationDataController& Controller, FName N
 }
 
 template <typename DataType, typename CurveClass, ERawCurveTrackTypes CurveType>
-void UAnimationBlueprintLibrary::AddCurveKeysInternal(UAnimSequence* AnimationSequence, FName CurveName, const TArray<float>& Times, const TArray<DataType>& KeyData)
+void UAnimationBlueprintLibrary::AddCurveKeysInternal(UAnimSequenceBase* AnimationSequenceBase, FName CurveName, const TArray<float>& Times, const TArray<DataType>& KeyData)
 {
 	checkf(Times.Num() == KeyData.Num(), TEXT("Not enough key data supplied"));
 
 	// Retrieve the curve by name
 	const FAnimationCurveIdentifier CurveId(CurveName, CurveType);
-	const CurveClass* Curve = static_cast<const CurveClass*>(AnimationSequence->GetDataModel()->FindCurve(CurveId));
+	const CurveClass* Curve = static_cast<const CurveClass*>(AnimationSequenceBase->GetDataModel()->FindCurve(CurveId));
 	if (Curve)
 	{
-		IAnimationDataController& Controller = AnimationSequence->GetController();
+		IAnimationDataController& Controller = AnimationSequenceBase->GetController();
 		SetControllerCurveKeys(Controller, CurveName, Times, KeyData);
 	}
 }
 
-bool UAnimationBlueprintLibrary::AddCurveInternal(UAnimSequence* AnimationSequence, FName CurveName, int32 CurveFlags, ERawCurveTrackTypes SupportedCurveType)
+bool UAnimationBlueprintLibrary::AddCurveInternal(UAnimSequenceBase* AnimationSequenceBase, FName CurveName, int32 CurveFlags, ERawCurveTrackTypes SupportedCurveType)
 {
 	// Add or retrieve the smart name
 	FAnimationCurveIdentifier CurveId(CurveName, SupportedCurveType);
 
 	bool bCurveAdded = false;
 
-	const FAnimCurveBase* ExistingCurve = AnimationSequence->GetDataModel()->FindCurve(CurveId);
+	const FAnimCurveBase* ExistingCurve = AnimationSequenceBase->GetDataModel()->FindCurve(CurveId);
 	if (ExistingCurve == nullptr)
 	{
-		IAnimationDataController& Controller = AnimationSequence->GetController();
+		IAnimationDataController& Controller = AnimationSequenceBase->GetController();
 		Controller.AddCurve(CurveId, CurveFlags);
-		bCurveAdded = AnimationSequence->GetDataModel()->FindCurve(CurveId) != nullptr;	
+		bCurveAdded = AnimationSequenceBase->GetDataModel()->FindCurve(CurveId) != nullptr;	
 	}
 	else
 	{
@@ -1516,12 +1514,12 @@ bool UAnimationBlueprintLibrary::AddCurveInternal(UAnimSequence* AnimationSequen
 	return bCurveAdded;
 }
 
-bool UAnimationBlueprintLibrary::RemoveCurveInternal(UAnimSequence* AnimationSequence, FName CurveName, ERawCurveTrackTypes SupportedCurveType)
+bool UAnimationBlueprintLibrary::RemoveCurveInternal(UAnimSequenceBase* AnimationSequenceBase, FName CurveName, ERawCurveTrackTypes SupportedCurveType)
 {
-	checkf(AnimationSequence != nullptr, TEXT("Invalid Animation Sequence ptr"));
+	checkf(AnimationSequenceBase != nullptr, TEXT("Invalid Animation Sequence ptr"));
 	bool bRemoved = false;
 	
-	IAnimationDataController& Controller = AnimationSequence->GetController();
+	IAnimationDataController& Controller = AnimationSequenceBase->GetController();
 
 	const FAnimationCurveIdentifier CurveId(CurveName, SupportedCurveType);
 	bRemoved = Controller.RemoveCurve(CurveId);
@@ -1556,11 +1554,11 @@ bool UAnimationBlueprintLibrary::DoesBoneNameExistInternal(USkeleton* Skeleton,
 	return Skeleton->GetReferenceSkeleton().FindBoneIndex(BoneName) != INDEX_NONE;
 }
 
-void UAnimationBlueprintLibrary::GetFloatKeys(UAnimSequence* AnimationSequence, FName CurveName, TArray<float>& Times, TArray<float>& Values)
+void UAnimationBlueprintLibrary::GetFloatKeys(UAnimSequenceBase* AnimationSequenceBase, FName CurveName, TArray<float>& Times, TArray<float>& Values)
 {
-	if (AnimationSequence)
+	if (AnimationSequenceBase)
 	{
-		GetCurveKeysInternal<float, FFloatCurve, ERawCurveTrackTypes::RCT_Float>(AnimationSequence, CurveName, Times, Values);
+		GetCurveKeysInternal<float, FFloatCurve, ERawCurveTrackTypes::RCT_Float>(AnimationSequenceBase, CurveName, Times, Values);
 	}
 	else
 	{
@@ -1568,11 +1566,11 @@ void UAnimationBlueprintLibrary::GetFloatKeys(UAnimSequence* AnimationSequence,
 	}
 }
 
-void UAnimationBlueprintLibrary::GetVectorKeys(UAnimSequence* AnimationSequence, FName CurveName, TArray<float>& Times, TArray<FVector>& Values)
+void UAnimationBlueprintLibrary::GetVectorKeys(UAnimSequenceBase* AnimationSequenceBase, FName CurveName, TArray<float>& Times, TArray<FVector>& Values)
 {
-	if (AnimationSequence)
+	if (AnimationSequenceBase)
 	{
-		GetCurveKeysInternal<FVector, FVectorCurve, ERawCurveTrackTypes::RCT_Vector>(AnimationSequence, CurveName, Times, Values);
+		GetCurveKeysInternal<FVector, FVectorCurve, ERawCurveTrackTypes::RCT_Vector>(AnimationSequenceBase, CurveName, Times, Values);
 	}
 	else
 	{
@@ -1580,11 +1578,11 @@ void UAnimationBlueprintLibrary::GetVectorKeys(UAnimSequence* AnimationSequence,
 	}
 }
 
-void UAnimationBlueprintLibrary::GetTransformationKeys(UAnimSequence* AnimationSequence, FName CurveName, TArray<float>& Times, TArray<FTransform>& Values)
+void UAnimationBlueprintLibrary::GetTransformationKeys(UAnimSequenceBase* AnimationSequenceBase, FName CurveName, TArray<float>& Times, TArray<FTransform>& Values)
 {
-	if (AnimationSequence)
+	if (AnimationSequenceBase)
 	{
-		GetCurveKeysInternal<FTransform, FTransformCurve, ERawCurveTrackTypes::RCT_Transform>(AnimationSequence, CurveName, Times, Values);
+		GetCurveKeysInternal<FTransform, FTransformCurve, ERawCurveTrackTypes::RCT_Transform>(AnimationSequenceBase, CurveName, Times, Values);
 	}
 	else
 	{
@@ -1592,14 +1590,35 @@ void UAnimationBlueprintLibrary::GetTransformationKeys(UAnimSequence* AnimationS
 	}
 }
 
+float UAnimationBlueprintLibrary::GetFloatValueAtTime(UAnimSequenceBase* AnimationSequenceBase, FName CurveName, float Time)
+{
+	if (AnimationSequenceBase)
+	{
+	    const FAnimationCurveIdentifier CurveId(CurveName, ERawCurveTrackTypes::RCT_Float);
+	    const FFloatCurve* Curve = AnimationSequenceBase->GetDataModel()->FindFloatCurve(CurveId);
+	    if (Curve)
+	    {
+		    return Curve->Evaluate(Time);
+	    }
+    
+	    UE_LOG(LogAnimationBlueprintLibrary, Warning, TEXT("Invalid Float Curve Name for given Animation Sequence"));
+	}
+	else
+	{
+		UE_LOG(LogAnimationBlueprintLibrary, Warning, TEXT("Invalid Animation Sequence for GetFloatValueAtTime"));
+	}
+
+	return 0.0f;
+}
+
 template <typename DataType, typename CurveClass, ERawCurveTrackTypes CurveType>
-void UAnimationBlueprintLibrary::GetCurveKeysInternal(UAnimSequence* AnimationSequence, FName CurveName, TArray<float>& Times, TArray<DataType>& KeyData)
+void UAnimationBlueprintLibrary::GetCurveKeysInternal(UAnimSequenceBase* AnimationSequenceBase, FName CurveName, TArray<float>& Times, TArray<DataType>& KeyData)
 {
-	checkf(AnimationSequence != nullptr, TEXT("Invalid Animation Sequence ptr"));
+	checkf(AnimationSequenceBase != nullptr, TEXT("Invalid Animation Sequence ptr"));
 	
 	// Retrieve the curve by name
 	const FAnimationCurveIdentifier CurveId(CurveName, CurveType);
-	const CurveClass* Curve = static_cast<const CurveClass*>(AnimationSequence->GetDataModel()->FindCurve(CurveId));
+	const CurveClass* Curve = static_cast<const CurveClass*>(AnimationSequenceBase->GetDataModel()->FindCurve(CurveId));
 
 	if (Curve)
 	{
@@ -1608,14 +1627,14 @@ void UAnimationBlueprintLibrary::GetCurveKeysInternal(UAnimSequence* AnimationSe
 	}
 }
 
-bool UAnimationBlueprintLibrary::DoesCurveExist(UAnimSequence* AnimationSequence, FName CurveName, ERawCurveTrackTypes CurveType)
+bool UAnimationBlueprintLibrary::DoesCurveExist(UAnimSequenceBase* AnimationSequenceBase, FName CurveName, ERawCurveTrackTypes CurveType)
 {
 	bool bExistingCurve = false;
 
-	if (AnimationSequence)
+	if (AnimationSequenceBase)
 	{
 		FAnimationCurveIdentifier CurveId(CurveName, CurveType);
-		const FAnimCurveBase* Curve = AnimationSequence->GetDataModel()->FindCurve(CurveId);
+		const FAnimCurveBase* Curve = AnimationSequenceBase->GetDataModel()->FindCurve(CurveId);
 		bExistingCurve = Curve != nullptr;
 	}
 	else
@@ -1626,14 +1645,14 @@ bool UAnimationBlueprintLibrary::DoesCurveExist(UAnimSequence* AnimationSequence
 	return bExistingCurve;
 }
 
-ERawCurveTrackTypes UAnimationBlueprintLibrary::RetrieveCurveTypeForCurve(const UAnimSequence* AnimationSequence, FName CurveName)
+ERawCurveTrackTypes UAnimationBlueprintLibrary::RetrieveCurveTypeForCurve(const UAnimSequenceBase* AnimationSequenceBase, FName CurveName)
 {
-	if(AnimationSequence->GetDataModel()->FindCurve(FAnimationCurveIdentifier(CurveName, ERawCurveTrackTypes::RCT_Float)))
+	if(AnimationSequenceBase->GetDataModel()->FindCurve(FAnimationCurveIdentifier(CurveName, ERawCurveTrackTypes::RCT_Float)))
 	{
 		return ERawCurveTrackTypes::RCT_Float;
 	}
 	
-	if(AnimationSequence->GetDataModel()->FindCurve(FAnimationCurveIdentifier(CurveName, ERawCurveTrackTypes::RCT_Transform)))
+	if(AnimationSequenceBase->GetDataModel()->FindCurve(FAnimationCurveIdentifier(CurveName, ERawCurveTrackTypes::RCT_Transform)))
 	{
 		return ERawCurveTrackTypes::RCT_Transform;
 	}
@@ -2543,12 +2562,12 @@ void UAnimationBlueprintLibrary::GetBonePosesForFrame(const UAnimSequenceBase* A
 	}
 }
 
-template void UAnimationBlueprintLibrary::AddCurveKeysInternal<float, FFloatCurve, ERawCurveTrackTypes::RCT_Float>(UAnimSequence* AnimationSequence, FName CurveName, const TArray<float>& Times, const TArray<float>& KeyData);
-template void UAnimationBlueprintLibrary::AddCurveKeysInternal<FVector, FVectorCurve, ERawCurveTrackTypes::RCT_Vector>(UAnimSequence* AnimationSequence, FName CurveName, const TArray<float>& Times, const TArray<FVector>& KeyData);
-template void UAnimationBlueprintLibrary::AddCurveKeysInternal<FTransform, FTransformCurve, ERawCurveTrackTypes::RCT_Transform>(UAnimSequence* AnimationSequence, FName CurveName, const TArray<float>& Times, const TArray<FTransform>& KeyData);
+template void UAnimationBlueprintLibrary::AddCurveKeysInternal<float, FFloatCurve, ERawCurveTrackTypes::RCT_Float>(UAnimSequenceBase* AnimationSequenceBase, FName CurveName, const TArray<float>& Times, const TArray<float>& KeyData);
+template void UAnimationBlueprintLibrary::AddCurveKeysInternal<FVector, FVectorCurve, ERawCurveTrackTypes::RCT_Vector>(UAnimSequenceBase* AnimationSequenceBase, FName CurveName, const TArray<float>& Times, const TArray<FVector>& KeyData);
+template void UAnimationBlueprintLibrary::AddCurveKeysInternal<FTransform, FTransformCurve, ERawCurveTrackTypes::RCT_Transform>(UAnimSequenceBase* AnimationSequenceBase, FName CurveName, const TArray<float>& Times, const TArray<FTransform>& KeyData);
 
-template void UAnimationBlueprintLibrary::GetCurveKeysInternal<float, FFloatCurve, ERawCurveTrackTypes::RCT_Float>(UAnimSequence* AnimationSequence, FName CurveName, TArray<float>& Times, TArray<float>& KeyData);
-template void UAnimationBlueprintLibrary::GetCurveKeysInternal<FVector, FVectorCurve, ERawCurveTrackTypes::RCT_Vector>(UAnimSequence* AnimationSequence, FName CurveName, TArray<float>& Times, TArray<FVector>& KeyData);
-template void UAnimationBlueprintLibrary::GetCurveKeysInternal<FTransform, FTransformCurve, ERawCurveTrackTypes::RCT_Transform>(UAnimSequence* AnimationSequence, FName CurveName, TArray<float>& Times, TArray<FTransform>& KeyData);
+template void UAnimationBlueprintLibrary::GetCurveKeysInternal<float, FFloatCurve, ERawCurveTrackTypes::RCT_Float>(UAnimSequenceBase* AnimationSequenceBase, FName CurveName, TArray<float>& Times, TArray<float>& KeyData);
+template void UAnimationBlueprintLibrary::GetCurveKeysInternal<FVector, FVectorCurve, ERawCurveTrackTypes::RCT_Vector>(UAnimSequenceBase* AnimationSequenceBase, FName CurveName, TArray<float>& Times, TArray<FVector>& KeyData);
+template void UAnimationBlueprintLibrary::GetCurveKeysInternal<FTransform, FTransformCurve, ERawCurveTrackTypes::RCT_Transform>(UAnimSequenceBase* AnimationSequenceBase, FName CurveName, TArray<float>& Times, TArray<FTransform>& KeyData);
 
 #undef LOCTEXT_NAMESPACE // "AnimationBlueprintLibrary"
\ No newline at end of file
diff --git a/Engine/Source/Editor/AnimationBlueprintLibrary/Public/AnimationBlueprintLibrary.h b/Engine/Source/Editor/AnimationBlueprintLibrary/Public/AnimationBlueprintLibrary.h
--- a/Engine/Source/Editor/AnimationBlueprintLibrary/Public/AnimationBlueprintLibrary.h
+++ b/Engine/Source/Editor/AnimationBlueprintLibrary/Public/AnimationBlueprintLibrary.h
@@ -84,7 +84,7 @@ public:
 
 	/** Retrieves the Names of the individual float curves for the given Animation Sequence */
 	UFUNCTION(BlueprintPure, Category = "AnimationBlueprintLibrary|Animation")
-	static void GetAnimationCurveNames(const UAnimSequence* AnimationSequence, ERawCurveTrackTypes CurveType, TArray<FName>& CurveNames);
+	static void GetAnimationCurveNames(const UAnimSequenceBase* AnimationSequenceBase, ERawCurveTrackTypes CurveType, TArray<FName>& CurveNames);
 
 	/** Retrieves the Raw Translation Animation Data for the given Animation Track Name and Animation Sequence */
 	UE_DEPRECATED(5.2, "GetRawTrackPositionData has been deprecated, use AnimationModel interface instead")
@@ -313,47 +313,47 @@ public:
 
 	/** Adds an Animation Curve by Type and Name to the given Animation Sequence */
 	UFUNCTION(BlueprintCallable, Category = "AnimationBlueprintLibrary|Curves")
-	static void AddCurve(UAnimSequence* AnimationSequence, FName CurveName, ERawCurveTrackTypes CurveType = ERawCurveTrackTypes::RCT_Float, bool bMetaDataCurve = false);
+	static void AddCurve(UAnimSequenceBase* AnimationSequenceBase, FName CurveName, ERawCurveTrackTypes CurveType = ERawCurveTrackTypes::RCT_Float, bool bMetaDataCurve = false);
 
 	/** Removes an Animation Curve by Name from the given Animation Sequence (Raw Animation Curves [Names] may not be removed from the Skeleton) */
 	UFUNCTION(BlueprintCallable, Category = "AnimationBlueprintLibrary|Curves")
-	static void RemoveCurve(UAnimSequence* AnimationSequence, FName CurveName, bool bRemoveNameFromSkeleton = false);
+	static void RemoveCurve(UAnimSequenceBase* AnimationSequenceBase, FName CurveName, bool bRemoveNameFromSkeleton = false);
 
 	/** Removes all Animation Curve Data from the given Animation Sequence (Raw Animation Curves [Names] may not be removed from the Skeleton) */
 	UFUNCTION(BlueprintCallable, Category = "AnimationBlueprintLibrary|Curves")
-	static void RemoveAllCurveData(UAnimSequence* AnimationSequence);
+	static void RemoveAllCurveData(UAnimSequenceBase* AnimationSequenceBase);
 
 	/** Adds a Transformation Key to the specified Animation Curve inside of the given Animation Sequence */
 	UFUNCTION(BlueprintCallable, Category = "AnimationBlueprintLibrary|Curves")
-	static void AddTransformationCurveKey(UAnimSequence* AnimationSequence, FName CurveName, const float Time, const FTransform& Transform);
+	static void AddTransformationCurveKey(UAnimSequenceBase* AnimationSequenceBase, FName CurveName, const float Time, const FTransform& Transform);
 
 	/** Adds a multiple of Transformation Keys to the specified Animation Curve inside of the given Animation Sequence */
 	UFUNCTION(BlueprintCallable, Category = "AnimationBlueprintLibrary|Curves")
-	static void AddTransformationCurveKeys(UAnimSequence* AnimationSequence, FName CurveName, const TArray<float>& Times, const TArray<FTransform>& Transforms);
+	static void AddTransformationCurveKeys(UAnimSequenceBase* AnimationSequenceBase, FName CurveName, const TArray<float>& Times, const TArray<FTransform>& Transforms);
 
 	/** Adds a Float Key to the specified Animation Curve inside of the given Animation Sequence */
 	UFUNCTION(BlueprintCallable, Category = "AnimationBlueprintLibrary|Curves")
-	static void AddFloatCurveKey(UAnimSequence* AnimationSequence, FName CurveName, const float Time, const float Value);
+	static void AddFloatCurveKey(UAnimSequenceBase* AnimationSequenceBase, FName CurveName, const float Time, const float Value);
 
 	/** Adds a multiple of Float Keys to the specified Animation Curve inside of the given Animation Sequence */
 	UFUNCTION(BlueprintCallable, Category = "AnimationBlueprintLibrary|Curves")
-	static void AddFloatCurveKeys(UAnimSequence* AnimationSequence, FName CurveName, const TArray<float>& Times, const TArray<float>& Values);
+	static void AddFloatCurveKeys(UAnimSequenceBase* AnimationSequenceBase, FName CurveName, const TArray<float>& Times, const TArray<float>& Values);
 
 	/** Adds a Vector Key to the specified Animation Curve inside of the given Animation Sequence */
 	UFUNCTION(BlueprintCallable, Category = "AnimationBlueprintLibrary|Curves")
-	static void AddVectorCurveKey(UAnimSequence* AnimationSequence, FName CurveName, const float Time, const FVector Vector);
+	static void AddVectorCurveKey(UAnimSequenceBase* AnimationSequenceBase, FName CurveName, const float Time, const FVector Vector);
 
 	/** Adds a multiple of Vector Keys to the specified Animation Curve inside of the given Animation Sequence */
 	UFUNCTION(BlueprintCallable, Category = "AnimationBlueprintLibrary|Curves")
-	static void AddVectorCurveKeys(UAnimSequence* AnimationSequence, FName CurveName, const TArray<float>& Times, const TArray<FVector>& Vectors);
+	static void AddVectorCurveKeys(UAnimSequenceBase* AnimationSequenceBase, FName CurveName, const TArray<float>& Times, const TArray<FVector>& Vectors);
 
 	// Curve helper functions
 	template <typename DataType, typename CurveClass, ERawCurveTrackTypes CurveType>
-	static void AddCurveKeysInternal(UAnimSequence* AnimationSequence, FName CurveName, const TArray<float>& Times, const TArray<DataType>& KeyData);
+	static void AddCurveKeysInternal(UAnimSequenceBase* AnimationSequenceBase, FName CurveName, const TArray<float>& Times, const TArray<DataType>& KeyData);
 
 	// Returns true if successfully added, false if it was already existing
-	static bool AddCurveInternal(UAnimSequence* AnimationSequence, FName CurveName, int32 CurveFlags, ERawCurveTrackTypes SupportedCurveType);
-	static bool RemoveCurveInternal(UAnimSequence* AnimationSequence, FName CurveName, ERawCurveTrackTypes SupportedCurveType);
+	static bool AddCurveInternal(UAnimSequenceBase* AnimationSequenceBase, FName CurveName, int32 CurveFlags, ERawCurveTrackTypes SupportedCurveType);
+	static bool RemoveCurveInternal(UAnimSequenceBase* AnimationSequenceBase, FName CurveName, ERawCurveTrackTypes SupportedCurveType);
 
 	/** Checks whether or not the given Bone Name exist on the Skeleton referenced by the given Animation Sequence */
 	UFUNCTION(BlueprintCallable, Category = "AnimationBlueprintLibrary|Skeleton")
@@ -366,18 +366,22 @@ public:
 
 	/** Retrieves, a multiple of, Float Key(s) from the specified Animation Curve inside of the given Animation Sequence */
 	UFUNCTION(BlueprintCallable, Category = "AnimationBlueprintLibrary|Curves")
-	static void GetFloatKeys(UAnimSequence* AnimationSequence, FName CurveName, TArray<float>& Times, TArray<float>& Values);
+	static void GetFloatKeys(UAnimSequenceBase* AnimationSequenceBase, FName CurveName, TArray<float>& Times, TArray<float>& Values);
 
 	/** Retrieves, a multiple of, Vector Key(s) from the specified Animation Curve inside of the given Animation Sequence */
 	UFUNCTION(BlueprintCallable, Category = "AnimationBlueprintLibrary|Curves")
-	static void GetVectorKeys(UAnimSequence* AnimationSequence, FName CurveName, TArray<float>& Times, TArray<FVector>& Values);
+	static void GetVectorKeys(UAnimSequenceBase* AnimationSequenceBase, FName CurveName, TArray<float>& Times, TArray<FVector>& Values);
 
 	/** Retrieves, a multiple of, Transformation Key(s) from the specified Animation Curve inside of the given Animation Sequence */
 	UFUNCTION(BlueprintCallable, Category = "AnimationBlueprintLibrary|Curves")
-	static void GetTransformationKeys(UAnimSequence* AnimationSequence, FName CurveName, TArray<float>& Times, TArray<FTransform>& Values);
+	static void GetTransformationKeys(UAnimSequenceBase* AnimationSequenceBase, FName CurveName, TArray<float>& Times, TArray<FTransform>& Values);
 
+	/** Retrieves an evaluated float value for a given time from the specified Animation Curve inside of the given Animation Sequence */
+	UFUNCTION(BlueprintCallable, Category = "AnimationBlueprintLibrary|Curves")
+	static float GetFloatValueAtTime(UAnimSequenceBase* AnimationSequenceBase, FName CurveName, float Time);
+	
 	template <typename DataType, typename CurveClass, ERawCurveTrackTypes CurveType>
-	static void GetCurveKeysInternal(UAnimSequence* AnimationSequence, FName CurveName, TArray<float>& Times, TArray<DataType>& KeyData);
+	static void GetCurveKeysInternal(UAnimSequenceBase* AnimationSequenceBase, FName CurveName, TArray<float>& Times, TArray<DataType>& KeyData);
 
 	/** Ensures that any curve names that do not exist on the NewSkeleton are added to it, in which case the SmartName on the actual curve itself will also be updated */
 	UFUNCTION(BlueprintCallable, Category = "AnimationBlueprintLibrary|Curves", meta=(DeprecatedFunction, DeprecationMessage="It is no longer necessary to copy curve names to the skeleton. If metadata is required to be updated, please use the metadata setting APIs."))
@@ -408,7 +412,7 @@ public:
 
 	/** Checks whether or not the given Curve Name exist on the Skeleton referenced by the given Animation Sequence */
 	UFUNCTION(BlueprintCallable, Category = "AnimationBlueprintLibrary|Curves")
-	static bool DoesCurveExist(UAnimSequence* AnimationSequence, FName CurveName, ERawCurveTrackTypes CurveType);
+	static bool DoesCurveExist(UAnimSequenceBase* AnimationSequenceBase, FName CurveName, ERawCurveTrackTypes CurveType);
 
 	UE_DEPRECATED(5.3, "This function is no longer used.")
 	static bool DoesSmartNameExist(UAnimSequence* AnimationSequence, FName Name) { return false; }
@@ -427,7 +431,7 @@ public:
 	UE_DEPRECATED(5.3, "This function is no longer used.")
 	static FName RetrieveContainerNameForCurve(const UAnimSequence* AnimationSequence, FName CurveName) { return NAME_None; }
 	
-	static ERawCurveTrackTypes RetrieveCurveTypeForCurve(const UAnimSequence* AnimationSequence, FName CurveName);
+	static ERawCurveTrackTypes RetrieveCurveTypeForCurve(const UAnimSequenceBase* AnimationSequenceBase, FName CurveName);
 
 	// MetaData
 
