Blueprint function to access root transforms from an AnimSequence

diff --git a/Engine/Source/Editor/AnimationBlueprintLibrary/Private/AnimationBlueprintLibrary.cpp b/Engine/Source/Editor/AnimationBlueprintLibrary/Private/AnimationBlueprintLibrary.cpp
--- a/Engine/Source/Editor/AnimationBlueprintLibrary/Private/AnimationBlueprintLibrary.cpp
+++ b/Engine/Source/Editor/AnimationBlueprintLibrary/Private/AnimationBlueprintLibrary.cpp
@@ -213,6 +213,17 @@ void UAnimationBlueprintLibrary::GetAnimationCurveNames(const UAnimSequenceBase*
 	}
 }
 
+FTransform UAnimationBlueprintLibrary::ExtractRootTrackTransform(const UAnimSequenceBase* AnimationSequenceBase, float Time)
+{
+	if (AnimationSequenceBase == nullptr)
+	{
+		UE_LOG(LogAnimationBlueprintLibrary, Warning, TEXT("Invalid Animation Sequence supplied for ExtractRootTrackTransform"));
+		return FTransform::Identity;
+	}
+	
+	return AnimationSequenceBase->ExtractRootTrackTransform(Time, nullptr);
+}
+
 const FRawAnimSequenceTrack& UAnimationBlueprintLibrary::GetRawAnimationTrackByName(const UAnimSequenceBase* AnimationSequenceBase, const FName TrackName)
 {
 	static FRawAnimSequenceTrack TempTrack;
diff --git a/Engine/Source/Editor/AnimationBlueprintLibrary/Public/AnimationBlueprintLibrary.h b/Engine/Source/Editor/AnimationBlueprintLibrary/Public/AnimationBlueprintLibrary.h
--- a/Engine/Source/Editor/AnimationBlueprintLibrary/Public/AnimationBlueprintLibrary.h
+++ b/Engine/Source/Editor/AnimationBlueprintLibrary/Public/AnimationBlueprintLibrary.h
@@ -86,6 +86,10 @@ public:
 	UFUNCTION(BlueprintPure, Category = "AnimationBlueprintLibrary|Animation")
 	static void GetAnimationCurveNames(const UAnimSequenceBase* AnimationSequenceBase, ERawCurveTrackTypes CurveType, TArray<FName>& CurveNames);
 
+	/** Gets the root transform from the raw animation at Time */
+	UFUNCTION(BlueprintCallable, Category="AnimationBlueprintLibrary|Animation")
+	static FTransform ExtractRootTrackTransform(const UAnimSequenceBase* AnimationSequenceBase, float Time);
+
 	/** Retrieves the Raw Translation Animation Data for the given Animation Track Name and Animation Sequence */
 	UE_DEPRECATED(5.2, "GetRawTrackPositionData has been deprecated, use AnimationModel interface instead")
 	UFUNCTION(BlueprintPure, Category = "AnimationBlueprintLibrary|RawTrackData")
