Add timecode display to the animation editor viewport.  This will include timecode, slate information, and rate as it was recorded on the animation sequence.  If the display settings are off in the viewport tool bar or no timecode attributes are on the current sequence then the timecode, rate and slate value are not shown.

diff --git a/Engine/Source/Editor/AnimationBlueprintLibrary/Private/AnimationBlueprintLibrary.cpp b/Engine/Source/Editor/AnimationBlueprintLibrary/Private/AnimationBlueprintLibrary.cpp
--- a/Engine/Source/Editor/AnimationBlueprintLibrary/Private/AnimationBlueprintLibrary.cpp
+++ b/Engine/Source/Editor/AnimationBlueprintLibrary/Private/AnimationBlueprintLibrary.cpp
@@ -2036,6 +2036,94 @@ bool UAnimationBlueprintLibrary::IsValidTimeInternal(const UAnimSequenceBase* An
 	return FMath::IsWithinInclusive(Time, 0.0f, AnimationSequenceBase->GetPlayLength());
 }
 
+namespace UE::AnimationBlueprintLibrary::Private
+{
+	struct FTimecodeBoneAttributeNames
+	{
+		FTimecodeBoneAttributeNames()
+		{
+			TCHourAttrName = TEXT("TCHour");
+			TCMinuteAttrName = TEXT("TCMinute");
+			TCSecondAttrName = TEXT("TCSecond");
+			TCFrameAttrName = TEXT("TCFrame");
+			TCSubframeAttrName = TEXT("TCSubframe");
+			TCRateAttrName = TEXT("TCRate");
+			TCSlateAttrName = TEXT("TCSlate");
+			if (const UAnimationSettings* AnimationSettings = UAnimationSettings::Get())
+			{
+				TCHourAttrName = AnimationSettings->BoneTimecodeCustomAttributeNameSettings.HourAttributeName;
+				TCMinuteAttrName = AnimationSettings->BoneTimecodeCustomAttributeNameSettings.MinuteAttributeName;
+				TCSecondAttrName = AnimationSettings->BoneTimecodeCustomAttributeNameSettings.SecondAttributeName;
+				TCFrameAttrName = AnimationSettings->BoneTimecodeCustomAttributeNameSettings.FrameAttributeName;
+				TCSubframeAttrName = AnimationSettings->BoneTimecodeCustomAttributeNameSettings.SubframeAttributeName;
+				TCRateAttrName = AnimationSettings->BoneTimecodeCustomAttributeNameSettings.RateAttributeName;
+				TCSlateAttrName = AnimationSettings->BoneTimecodeCustomAttributeNameSettings.TakenameAttributeName;
+			}
+		}
+
+		bool Contains(const FName AttributeName) const
+		{
+			return AttributeName == TCHourAttrName
+				|| AttributeName == TCMinuteAttrName
+				|| AttributeName == TCSecondAttrName
+				|| AttributeName == TCFrameAttrName
+				|| AttributeName == TCSubframeAttrName
+				|| AttributeName == TCRateAttrName
+				|| AttributeName == TCSlateAttrName;
+		}
+
+		FName TCHourAttrName;
+		FName TCMinuteAttrName;
+		FName TCSecondAttrName;
+		FName TCFrameAttrName;
+		FName TCSubframeAttrName;
+		FName TCRateAttrName;
+		FName TCSlateAttrName;
+	};
+}
+
+FName UAnimationBlueprintLibrary::FindBoneNameWithTimecodeAttributes(const UAnimSequenceBase* AnimationSequenceBase)
+{
+	if (!AnimationSequenceBase || !AnimationSequenceBase->GetSkeleton())
+	{
+		return NAME_None;
+	}
+
+	const IAnimationDataModel* AnimDataModel = AnimationSequenceBase->GetDataModel();
+	if (!AnimDataModel || !AnimDataModel->HasBeenPopulated())
+	{
+		return NAME_None;
+	}
+
+	using namespace UE::AnimationBlueprintLibrary::Private;
+	FTimecodeBoneAttributeNames TimecodeBoneAttributeNames;
+
+	const TArray<FMeshBoneInfo>& BoneInfos = AnimationSequenceBase->GetSkeleton()->GetReferenceSkeleton().GetRefBoneInfo();
+	for (FMeshBoneInfo BoneInfo : BoneInfos)
+	{
+		if (!AnimDataModel->IsValidBoneTrackName(BoneInfo.Name))
+		{
+			continue;
+		}
+		TArray<const FAnimatedBoneAttribute*> BoneAttributes;
+		AnimDataModel->GetAttributesForBone(BoneInfo.Name, BoneAttributes);
+		for (const FAnimatedBoneAttribute* BoneAttribute : BoneAttributes)
+		{
+			if (!BoneAttribute)
+			{
+				continue;
+			}
+			const FName& BoneAttributeName = BoneAttribute->Identifier.GetName();
+
+			if (TimecodeBoneAttributeNames.Contains(BoneAttributeName))
+			{
+				return BoneInfo.Name;
+			}
+		}
+	}
+	return NAME_None;
+}
+
 bool UAnimationBlueprintLibrary::EvaluateRootBoneTimecodeAttributesAtTime(const UAnimSequenceBase* AnimationSequenceBase, const float EvalTime, FQualifiedFrameTime& OutQualifiedFrameTime)
 {
 	if (!AnimationSequenceBase || !AnimationSequenceBase->GetSkeleton())
@@ -2049,33 +2137,34 @@ bool UAnimationBlueprintLibrary::EvaluateRootBoneTimecodeAttributesAtTime(const
 		return false;
 	}
 
+	FString OutSlate;
 	const FName RootBoneName = AnimationSequenceBase->GetSkeleton()->GetReferenceSkeleton().GetBoneName(0);
-	if (!AnimDataModel->IsValidBoneTrackName(RootBoneName))
+	return EvaluateBoneTimecodeAndSlateAttributesAtTime(RootBoneName, AnimationSequenceBase, EvalTime, OutQualifiedFrameTime, OutSlate);
+}
+
+bool UAnimationBlueprintLibrary::EvaluateBoneTimecodeAndSlateAttributesAtTime(const FName BoneName, const UAnimSequenceBase* AnimationSequenceBase, const float EvalTime, FQualifiedFrameTime& OutQualifiedFrameTime, FString& OutSlate)
+{
+	if (!AnimationSequenceBase || !AnimationSequenceBase->GetSkeleton())
 	{
 		return false;
 	}
 
-	TArray<const FAnimatedBoneAttribute*> RootBoneAttributes;
-	AnimDataModel->GetAttributesForBone(RootBoneName, RootBoneAttributes);
-
-	FName TCHourAttrName(TEXT("TCHour"));
-	FName TCMinuteAttrName(TEXT("TCMinute"));
-	FName TCSecondAttrName(TEXT("TCSecond"));
-	FName TCFrameAttrName(TEXT("TCFrame"));
-	FName TCSubframeAttrName(TEXT("TCSubframe"));
-	FName TCRateAttrName(TEXT("TCRate"));
+	const IAnimationDataModel* AnimDataModel = AnimationSequenceBase->GetDataModel();
+	if (!AnimDataModel || !AnimDataModel->HasBeenPopulated())
+	{
+		return false;
+	}
 
-	if (const UAnimationSettings* AnimationSettings = UAnimationSettings::Get())
+	if (!AnimDataModel->IsValidBoneTrackName(BoneName))
 	{
-		TCHourAttrName = AnimationSettings->BoneTimecodeCustomAttributeNameSettings.HourAttributeName;
-		TCMinuteAttrName = AnimationSettings->BoneTimecodeCustomAttributeNameSettings.MinuteAttributeName;
-		TCSecondAttrName = AnimationSettings->BoneTimecodeCustomAttributeNameSettings.SecondAttributeName;
-		TCFrameAttrName = AnimationSettings->BoneTimecodeCustomAttributeNameSettings.FrameAttributeName;
-		TCSubframeAttrName = AnimationSettings->BoneTimecodeCustomAttributeNameSettings.SubframeAttributeName;
-		TCRateAttrName = AnimationSettings->BoneTimecodeCustomAttributeNameSettings.RateAttributeName;
+		return false;
 	}
 
-	const TArray<FName> TimecodeBoneAttributeNames = { TCHourAttrName, TCMinuteAttrName, TCSecondAttrName, TCFrameAttrName, TCSubframeAttrName, TCRateAttrName };
+	TArray<const FAnimatedBoneAttribute*> RootBoneAttributes;
+	AnimDataModel->GetAttributesForBone(BoneName, RootBoneAttributes);
+
+	using namespace UE::AnimationBlueprintLibrary::Private;
+	FTimecodeBoneAttributeNames TimecodeBoneAttributeNames;
 
 	bool bHasTimecodeBoneAttributes = false;
 
@@ -2137,38 +2226,42 @@ bool UAnimationBlueprintLibrary::EvaluateRootBoneTimecodeAttributesAtTime(const
 			continue;
 		}
 
-		if (BoneAttributeName.IsEqual(TCHourAttrName))
+		if (BoneAttributeName.IsEqual(TimecodeBoneAttributeNames.TCHourAttrName))
 		{
 			Timecode.Hours = IntValue;
 			bHasTimecodeBoneAttributes = true;
 		}
-		else if (BoneAttributeName.IsEqual(TCMinuteAttrName))
+		else if (BoneAttributeName.IsEqual(TimecodeBoneAttributeNames.TCMinuteAttrName))
 		{
 			Timecode.Minutes = IntValue;
 			bHasTimecodeBoneAttributes = true;
 		}
-		else if (BoneAttributeName.IsEqual(TCSecondAttrName))
+		else if (BoneAttributeName.IsEqual(TimecodeBoneAttributeNames.TCSecondAttrName))
 		{
 			Timecode.Seconds = IntValue;
 			bHasTimecodeBoneAttributes = true;
 		}
-		else if (BoneAttributeName.IsEqual(TCFrameAttrName))
+		else if (BoneAttributeName.IsEqual(TimecodeBoneAttributeNames.TCFrameAttrName))
 		{
 			Timecode.Frames = IntValue;
 			bHasTimecodeBoneAttributes = true;
 		}
-		else if (BoneAttributeName.IsEqual(TCSubframeAttrName))
+		else if (BoneAttributeName.IsEqual(TimecodeBoneAttributeNames.TCSubframeAttrName))
 		{
 			SubFrame = FloatValue;
 			bHasTimecodeBoneAttributes = true;
 		}
-		else if (BoneAttributeName.IsEqual(TCRateAttrName))
+		else if (BoneAttributeName.IsEqual(TimecodeBoneAttributeNames.TCRateAttrName))
 		{
 			TimecodeRateAsString = StringValue;
 			TimecodeRateAsDecimal = FloatValue;
 			// Don't consider this attribute when determining whether timecode value attributes are
 			// present, since it can't be useful on its own.
 		}
+		else if (BoneAttributeName.IsEqual(TimecodeBoneAttributeNames.TCSlateAttrName))
+		{
+			OutSlate = StringValue;
+		}
 	}
 
 	if (!bHasTimecodeBoneAttributes)
@@ -2581,4 +2674,4 @@ template void UAnimationBlueprintLibrary::GetCurveKeysInternal<float, FFloatCurv
 template void UAnimationBlueprintLibrary::GetCurveKeysInternal<FVector, FVectorCurve, ERawCurveTrackTypes::RCT_Vector>(UAnimSequenceBase* AnimationSequenceBase, FName CurveName, TArray<float>& Times, TArray<FVector>& KeyData);
 template void UAnimationBlueprintLibrary::GetCurveKeysInternal<FTransform, FTransformCurve, ERawCurveTrackTypes::RCT_Transform>(UAnimSequenceBase* AnimationSequenceBase, FName CurveName, TArray<float>& Times, TArray<FTransform>& KeyData);
 
-#undef LOCTEXT_NAMESPACE // "AnimationBlueprintLibrary"
\ No newline at end of file
+#undef LOCTEXT_NAMESPACE // "AnimationBlueprintLibrary"
diff --git a/Engine/Source/Editor/AnimationBlueprintLibrary/Public/AnimationBlueprintLibrary.h b/Engine/Source/Editor/AnimationBlueprintLibrary/Public/AnimationBlueprintLibrary.h
--- a/Engine/Source/Editor/AnimationBlueprintLibrary/Public/AnimationBlueprintLibrary.h
+++ b/Engine/Source/Editor/AnimationBlueprintLibrary/Public/AnimationBlueprintLibrary.h
@@ -543,6 +543,13 @@ public:
 
 	static bool IsValidTimeInternal(const UAnimSequenceBase* AnimationSequenceBase, const float Time);
 
+	/**
+	 * Given an animation sequence determine if any bones in that animation sequence has timecode attribute data.
+	 *
+	 * @return The bone name if found; otherwise none value will be returned.
+	 */
+	static FName FindBoneNameWithTimecodeAttributes(const UAnimSequenceBase* AnimSequenceBase);
+
 	/** Evaluates timecode attributes (e.g. "TCFrame", "TCSecond", etc.) of the root bone and returns the resulting qualified frame time.
 	 *
 	 *  @param AnimationSequenceBase: Anim sequence for which to evaluate the root bone attributes.
@@ -555,6 +562,20 @@ public:
 	UFUNCTION(BlueprintPure, Category = "AnimationBlueprintLibrary|Helpers")
 	static bool EvaluateRootBoneTimecodeAttributesAtTime(const UAnimSequenceBase* AnimationSequenceBase, const float EvalTime, FQualifiedFrameTime& OutQualifiedFrameTime);
 
+	/** Evaluates timecode attributes (e.g. "TCFrame", "TCSecond", etc.) and TCSlate of the root bone and returns the resulting qualified frame time.
+	 *
+	 *  @param AnimationSequenceBase: Anim sequence for which to evaluate the root bone attributes.
+	 *  @param EvalTime: Time (in seconds) at which to evaluate the timecode bone attributes.
+	 *  @param OutQualifiedFrameTime: Resulting qualified frame time from evaluation. If the anim sequence has an import file frame rate
+	 *      set, then that will be used as the frame rate of the qualified frame time. Otherwise, the sampling frame rate of the anim
+	 *      sequence is used. If no timecode attributes are present on the bone or if none can be evaluated, the passed object will not be modified.
+	 *  @param OutSlate:  The slate name (if any).  If so TCSlate exists on the root bone then an empty string is returned.
+	 *
+	 *  @return: true if the root bone had timecode attributes that could be evaluated and a qualified frame time was set, or false otherwise.
+	 */
+	UFUNCTION(BlueprintPure, Category = "AnimationBlueprintLibrary|Helpers")
+	static bool EvaluateBoneTimecodeAndSlateAttributesAtTime(const FName BoneName, const UAnimSequenceBase* AnimationSequenceBase, const float EvalTime, FQualifiedFrameTime& OutQualifiedFrameTime, FString& Slate);
+
 	/** Evaluates the subframe timecode attribute (e.g. "TCSubframe") of the root bone and returns the resulting value.
 	 *
 	 *  Since the subframe component of FFrameTime is clamped to the range [0.0, 1.0), it cannot accurately represent the use
diff --git a/Engine/Source/Editor/Persona/Private/AnimViewportShowCommands.cpp b/Engine/Source/Editor/Persona/Private/AnimViewportShowCommands.cpp
--- a/Engine/Source/Editor/Persona/Private/AnimViewportShowCommands.cpp
+++ b/Engine/Source/Editor/Persona/Private/AnimViewportShowCommands.cpp
@@ -75,6 +75,8 @@ void FAnimViewportShowCommands::RegisterCommands()
 	UI_COMMAND(ShowNotificationVisualizations, "Show Notification", "Display animation notifications visualizations in viewport", EUserInterfaceActionType::ToggleButton, FInputChord());
 	UI_COMMAND(ShowRootMotionVisualization, "Show Root Motion", "Display root motion visualizations in viewport", EUserInterfaceActionType::ToggleButton, FInputChord());
 	UI_COMMAND(ShowAssetUserDataVisualizations, "Show Meta Data", "Display animation meta data visualizations in viewport", EUserInterfaceActionType::ToggleButton, FInputChord());
+
+	UI_COMMAND(ShowTimecode, "Show Timecode", "Display timecode information associated with animation data (if exists).", EUserInterfaceActionType::ToggleButton, FInputChord());
 }
 
-#undef LOCTEXT_NAMESPACE
\ No newline at end of file
+#undef LOCTEXT_NAMESPACE
diff --git a/Engine/Source/Editor/Persona/Private/AnimViewportShowCommands.h b/Engine/Source/Editor/Persona/Private/AnimViewportShowCommands.h
--- a/Engine/Source/Editor/Persona/Private/AnimViewportShowCommands.h
+++ b/Engine/Source/Editor/Persona/Private/AnimViewportShowCommands.h
@@ -162,6 +162,9 @@ public:
 	/** Show AssetUserData Visualizations */
 	TSharedPtr< FUICommandInfo > ShowAssetUserDataVisualizations;
 
+	/** Show Animation Timecode Values */
+	TSharedPtr< FUICommandInfo > ShowTimecode;
+
 public:
 	/** Registers our commands with the binding system */
 	virtual void RegisterCommands() override;
diff --git a/Engine/Source/Editor/Persona/Private/AnimationEditorPreviewScene.cpp b/Engine/Source/Editor/Persona/Private/AnimationEditorPreviewScene.cpp
--- a/Engine/Source/Editor/Persona/Private/AnimationEditorPreviewScene.cpp
+++ b/Engine/Source/Editor/Persona/Private/AnimationEditorPreviewScene.cpp
@@ -1292,4 +1292,13 @@ void FAnimationEditorPreviewScene::AddReferencedObjects( FReferenceCollector& Co
 	Collector.AddReferencedObjects(AdditionalMeshes);
 }
 
+TOptional<float> FAnimationEditorPreviewScene::GetCurrentTime() const
+{
+	if (SkeletalMeshComponent && SkeletalMeshComponent->PreviewInstance)
+	{
+		return SkeletalMeshComponent->PreviewInstance->GetCurrentTime();
+	}
+	return {};
+}
+
 #undef LOCTEXT_NAMESPACE
diff --git a/Engine/Source/Editor/Persona/Private/AnimationEditorPreviewScene.h b/Engine/Source/Editor/Persona/Private/AnimationEditorPreviewScene.h
--- a/Engine/Source/Editor/Persona/Private/AnimationEditorPreviewScene.h
+++ b/Engine/Source/Editor/Persona/Private/AnimationEditorPreviewScene.h
@@ -334,6 +334,21 @@ public:
 		return bIsBeingConstructed;
 	}
 
+	/** Toggles state on showing Timecode value in the viewport */
+	void ToggleShowTimecode()
+	{
+		bShowTimecode = !bShowTimecode;
+	}
+
+	/** Should we show the timecode in the viewport. */
+	bool IsShowTimecode() const
+	{
+		return bShowTimecode;
+	}
+
+	/** Get the current time of the preview instance for the current animation sequence. */
+	TOptional<float> GetCurrentTime() const;
+
 private:
 	/** Set preview mesh internal use only. The mesh should be verified by now. */
 	void SetPreviewMeshInternal(USkeletalMesh* NewPreviewMesh);
@@ -483,4 +498,7 @@ private:
 
 	/** True during initial creation, so some code can be skipped */
 	bool bIsBeingConstructed = false;
+
+	/** Should we display timecode in the preview viewport. */
+	bool bShowTimecode = true;
 };
diff --git a/Engine/Source/Editor/Persona/Private/AnimationEditorViewportClient.cpp b/Engine/Source/Editor/Persona/Private/AnimationEditorViewportClient.cpp
--- a/Engine/Source/Editor/Persona/Private/AnimationEditorViewportClient.cpp
+++ b/Engine/Source/Editor/Persona/Private/AnimationEditorViewportClient.cpp
@@ -50,6 +50,7 @@
 #include "Animation/AnimNotifies/AnimNotifyState.h"
 #include "Animation/AnimMontage.h"
 #include "Animation/AnimSequenceHelpers.h"
+#include "AnimationBlueprintLibrary.h"
 
 namespace {
 	static const float AnimationEditorViewport_RotateSpeed = 0.02f;
@@ -723,6 +724,24 @@ void FAnimationViewportClient::Tick(float DeltaSeconds)
 	FEditorViewportClient::Tick(DeltaSeconds);
 
 	GetAnimPreviewScene()->FlagTickable();
+
+	TimecodeDisplay.Reset();
+	if (GetAnimPreviewScene()->IsShowTimecode())
+	{
+		UAnimationAsset* AnimationAsset = GetAnimPreviewScene()->GetPreviewAnimationAsset();
+		if (UAnimSequence* AnimSequence = Cast<UAnimSequence>(AnimationAsset))
+		{
+			FName BoneName = UAnimationBlueprintLibrary::FindBoneNameWithTimecodeAttributes(AnimSequence);
+
+			FString SlateName;
+			FQualifiedFrameTime QualifiedFrameTime;
+			TOptional<float> PlayPosition = GetAnimPreviewScene()->GetCurrentTime();
+			if (PlayPosition && UAnimationBlueprintLibrary::EvaluateBoneTimecodeAndSlateAttributesAtTime(BoneName, AnimSequence, *PlayPosition, QualifiedFrameTime, SlateName))
+			{
+				TimecodeDisplay = {QualifiedFrameTime, SlateName};
+			}
+		}
+	}
 }
 
 void FAnimationViewportClient::HandlePreviewScenePreTick()
@@ -1329,6 +1348,12 @@ FText FAnimationViewportClient::GetDisplayInfo(bool bDisplayAllInfo) const
 		}
 	}
 
+	if (TimecodeDisplay)
+	{
+		TextValue = ConcatenateLine(TextValue, FText::Format(LOCTEXT("TimecodeInfo", "Timecode: {0}"), FText::FromString(TimecodeDisplay->QualifiedTime.ToTimecode().ToString())));
+		TextValue = ConcatenateLine(TextValue, FText::Format(LOCTEXT("SlateName", "Slate: {0}"), FText::FromString(TimecodeDisplay->Slate)));
+		TextValue = ConcatenateLine(TextValue, FText::Format(LOCTEXT("Rate", "Rate: {0}"), FText::AsNumber(TimecodeDisplay->QualifiedTime.Rate.AsDecimal())));
+	}
 	return TextValue;
 }
 void FAnimationViewportClient::DrawNodeDebugLines(TArray<FText>& Lines, FCanvas* Canvas, FSceneView* View)
diff --git a/Engine/Source/Editor/Persona/Private/SAnimViewportToolBar.cpp b/Engine/Source/Editor/Persona/Private/SAnimViewportToolBar.cpp
--- a/Engine/Source/Editor/Persona/Private/SAnimViewportToolBar.cpp
+++ b/Engine/Source/Editor/Persona/Private/SAnimViewportToolBar.cpp
@@ -880,6 +880,19 @@ TSharedRef<SWidget> SAnimViewportToolBar::GenerateCharacterMenu() const
 						FNewToolMenuChoice(FNewMenuDelegate::CreateRaw(ContextThis, &SAnimViewportToolBar::FillCharacterAdvancedMenu)));
 				}
 			}));
+
+			Section.AddDynamicEntry("Timecode", FNewToolMenuSectionDelegate::CreateLambda([](FToolMenuSection& InSection)
+			{
+				UAnimViewportToolBarToolMenuContext* Context = InSection.FindContext<UAnimViewportToolBarToolMenuContext>();
+				const SAnimViewportToolBar* ContextThis = Context ? Context->AnimViewportToolBar.Pin().Get() : nullptr;
+				if (ContextThis)
+				{
+					InSection.AddSubMenu(TEXT("TimecodeSubMenu"),
+						LOCTEXT("CharacterMenu_TimecodeSubMenu", "Timecode"),
+						LOCTEXT("CharacterMenu_TimecodeSubMenuToolTip", "Timecode options"),
+						FNewToolMenuChoice(FNewMenuDelegate::CreateRaw(ContextThis, &SAnimViewportToolBar::FillCharacterTimecodeMenu)));
+				}
+			}));
 		}
 	}
 
@@ -893,6 +906,16 @@ TSharedRef<SWidget> SAnimViewportToolBar::GenerateCharacterMenu() const
 	return UToolMenus::Get()->GenerateWidget(MenuName, MenuContext);
 }
 
+void SAnimViewportToolBar::FillCharacterTimecodeMenu(FMenuBuilder& MenuBuilder) const
+{
+	const FAnimViewportShowCommands& Actions = FAnimViewportShowCommands::Get();
+	MenuBuilder.BeginSection("Timecode", LOCTEXT("Timecode_Label", "Timecode"));
+	{
+		MenuBuilder.AddMenuEntry( Actions.ShowTimecode );
+	}
+	MenuBuilder.EndSection();
+}
+
 void SAnimViewportToolBar::FillCharacterAdvancedMenu(FMenuBuilder& MenuBuilder) const
 {
 	const FAnimViewportShowCommands& Actions = FAnimViewportShowCommands::Get();
diff --git a/Engine/Source/Editor/Persona/Private/SAnimViewportToolBar.h b/Engine/Source/Editor/Persona/Private/SAnimViewportToolBar.h
--- a/Engine/Source/Editor/Persona/Private/SAnimViewportToolBar.h
+++ b/Engine/Source/Editor/Persona/Private/SAnimViewportToolBar.h
@@ -100,6 +100,11 @@ private:
 	*/
 	void FillCharacterClothingMenu(FMenuBuilder& MenuBuilder);
 
+	/**
+	* Generate the Character -> Timecode menu
+	*/
+	void FillCharacterTimecodeMenu(FMenuBuilder& MenuBuilder) const;
+
 	/**
 	 * Generates the toolbar LOD menu content 
 	 */
diff --git a/Engine/Source/Editor/Persona/Private/SAnimationEditorViewport.cpp b/Engine/Source/Editor/Persona/Private/SAnimationEditorViewport.cpp
--- a/Engine/Source/Editor/Persona/Private/SAnimationEditorViewport.cpp
+++ b/Engine/Source/Editor/Persona/Private/SAnimationEditorViewport.cpp
@@ -1007,6 +1007,16 @@ void SAnimationEditorViewportTabBody::BindCommands()
 
 	CommandList.EndGroup();
 
+	CommandList.BeginGroup(TEXT("TimecodeSettings"));
+
+	CommandList.MapAction(
+		ViewportShowMenuCommands.ShowTimecode,
+		FExecuteAction::CreateSP(this, &SAnimationEditorViewportTabBody::OnToggleShowTimecode),
+		FCanExecuteAction(),
+		FIsActionChecked::CreateSP(this, &SAnimationEditorViewportTabBody::IsShowTimecode));
+
+	CommandList.EndGroup();
+
 
 	GetPreviewScene()->RegisterOnSelectedLODChanged(FOnSelectedLODChanged::CreateSP(this, &SAnimationEditorViewportTabBody::OnLODModelChanged));
 	//Bind LOD preview menu commands
@@ -1624,6 +1634,16 @@ bool SAnimationEditorViewportTabBody::IsShowPreviewMeshEnabled() const
 	return false;
 }
 
+void SAnimationEditorViewportTabBody::OnToggleShowTimecode()
+{
+	GetPreviewScene()->ToggleShowTimecode();
+}
+
+bool SAnimationEditorViewportTabBody::IsShowTimecode() const
+{
+	return GetPreviewScene()->IsShowTimecode();
+}
+
 void SAnimationEditorViewportTabBody::UseInGameBound()
 {
 	ForEachDebugMesh([](UDebugSkelMeshComponent* PreviewMeshComponent)
@@ -2757,7 +2777,6 @@ void SAnimationEditorViewportTabBody::AddSkinWeightProfileNotification()
 void SAnimationEditorViewportTabBody::HandleFocusCamera()
 {
 	TSharedRef<FAnimationViewportClient> AnimViewportClient = StaticCastSharedRef<FAnimationViewportClient>(LevelViewportClient.ToSharedRef());
-	// AnimViewportClient->SetCameraFollowMode(EAnimationViewportCameraFollowMode::None);
 	AnimViewportClient->FocusViewportOnPreviewMesh(false);
 }
 
diff --git a/Engine/Source/Editor/Persona/Private/SAnimationEditorViewport.h b/Engine/Source/Editor/Persona/Private/SAnimationEditorViewport.h
--- a/Engine/Source/Editor/Persona/Private/SAnimationEditorViewport.h
+++ b/Engine/Source/Editor/Persona/Private/SAnimationEditorViewport.h
@@ -110,7 +110,7 @@ protected:
 	virtual void OnDragEnter(const FGeometry& MyGeometry, const FDragDropEvent& DragDropEvent) override;
 	virtual void OnDragLeave(const FDragDropEvent& DragDropEvent) override;
 	virtual FReply OnDrop(const FGeometry& MyGeometry, const FDragDropEvent& DragDropEvent) override;
-	
+
 protected:
 	// Viewport client
 	TSharedPtr<class FAnimationViewportClient> LevelViewportClient;
@@ -229,7 +229,7 @@ public:
 	virtual void AddOverlayWidget( TSharedRef<SWidget> InOverlaidWidget, int32 ZOrder = INDEX_NONE ) override;
 	virtual void RemoveOverlayWidget( TSharedRef<SWidget> InOverlaidWidget ) override;
 
-	
+
 	/** SWidget interface */
 	virtual FReply OnKeyDown(const FGeometry& MyGeometry, const FKeyEvent& InKeyEvent) override;
 
@@ -423,6 +423,12 @@ private:
 	/** Function to check whether mesh info is displayed or not */
 	bool IsShowingMeshInfo(int32 DisplayInfoMode) const;
 
+	/** Function to show/hide timecode display. */
+	void OnToggleShowTimecode();
+
+	/** Function to check whether timecode is displayed or not */
+	bool IsShowTimecode() const;
+
 	/** Toggles floor alignment in the preview scene */
 	void OnToggleAutoAlignFloor();
 
diff --git a/Engine/Source/Editor/Persona/Public/AnimationEditorViewportClient.h b/Engine/Source/Editor/Persona/Public/AnimationEditorViewportClient.h
--- a/Engine/Source/Editor/Persona/Public/AnimationEditorViewportClient.h
+++ b/Engine/Source/Editor/Persona/Public/AnimationEditorViewportClient.h
@@ -497,6 +497,11 @@ private:
 	void HandlePreviewScenePostTick();
 
 private:
+	struct FTimecodeDisplayInfo
+	{
+		FQualifiedFrameTime QualifiedTime;
+		FString	  Slate;
+	};
 
 	/** @return array of AssetUserData interfaces from editable objects on associated asset toolkit. */
 	TArray<IInterface_AssetUserData*> GetEditedObjectsWithAssetUserData() const;
@@ -530,4 +535,7 @@ private:
 
 	/** True when the preview animation should resume playing upon finishing tracking */
 	bool bResumeAfterTracking;
+
+	/** Timecode/slate information from current animation sequence. */
+	TOptional<FTimecodeDisplayInfo> TimecodeDisplay;
 };
