Fixup EvaluateAnimationBlueprintWithInputPose

diff --git a/Engine/Source/Editor/AnimationBlueprintLibrary/Private/AnimPose.cpp b/Engine/Source/Editor/AnimationBlueprintLibrary/Private/AnimPose.cpp
--- a/Engine/Source/Editor/AnimationBlueprintLibrary/Private/AnimPose.cpp
+++ b/Engine/Source/Editor/AnimationBlueprintLibrary/Private/AnimPose.cpp
@@ -463,6 +463,7 @@ void UAnimPoseExtensions::EvaluateAnimationBlueprintWithInputPose(const FAnimPos
 						Component->SetAnimInstanceClass(AnimationBlueprint->GetAnimBlueprintGeneratedClass());
 
 						PreviewScene.AddComponent(Component, FTransform::Identity);
+						Component->InitAnim(true);
 			
 						if (UAnimInstance* AnimInstance = Component->GetAnimInstance())
 						{
@@ -472,6 +473,7 @@ void UAnimPoseExtensions::EvaluateAnimationBlueprintWithInputPose(const FAnimPos
 								InputNode->CachedInputPose.SetBoneContainer(&BoneContainer);
 								InputNode->CachedInputCurve.InitFrom(BoneContainer);
 								InputNode->CachedInputPose.ResetToRefPose();
+								InputNode->bIsCachedInputPoseInitialized = true;
 
 								// Copy bone transform from input pose using skeleton index mapping
 								for (FCompactPoseBoneIndex CompactIndex : InputNode->CachedInputPose.ForEachBoneIndex())
@@ -489,7 +491,7 @@ void UAnimPoseExtensions::EvaluateAnimationBlueprintWithInputPose(const FAnimPos
 					
 								OutPose.Init(AnimInstance->GetRequiredBones());
 
-								Component->InitAnim(true);
+								Component->TickAnimation(0.f, true);
 								Component->RefreshBoneTransforms();
 								const TArray<FTransform>& LocalSpaceTransforms = Component->GetBoneSpaceTransforms();
 
