addition of RemoveBoneSelectiveAnimation

#rnx

diff --git a/Engine/Source/Editor/AnimationBlueprintLibrary/Private/AnimationBlueprintLibrary.cpp b/Engine/Source/Editor/AnimationBlueprintLibrary/Private/AnimationBlueprintLibrary.cpp
--- a/Engine/Source/Editor/AnimationBlueprintLibrary/Private/AnimationBlueprintLibrary.cpp
+++ b/Engine/Source/Editor/AnimationBlueprintLibrary/Private/AnimationBlueprintLibrary.cpp
@@ -2441,6 +2441,79 @@ void UAnimationBlueprintLibrary::RemoveBoneAnimation(UAnimSequence* AnimationSeq
 	}
 }
 
+
+void UAnimationBlueprintLibrary::RemoveBoneSelectiveAnimation(UAnimSequence* AnimationSequence, FName BoneName, TArray<FName> ChildrenExcluded, bool bIncludeChildren /*= true*/, bool bExcludeRecursively /*= false*/, bool bFinalize /*= true*/)
+{
+	if (AnimationSequence)
+	{
+		TArray<FName> TrackNames;
+		AnimationSequence->GetDataModel()->GetBoneTrackNames(TrackNames);
+		
+		if (TrackNames.Contains(BoneName))
+		{
+			TArray<FName> TracksToRemove;
+			TracksToRemove.Add(BoneName);
+
+
+			// remove all children if required
+			if (bIncludeChildren)
+			{
+				USkeleton* Skeleton = AnimationSequence->GetSkeleton();
+				if (Skeleton)
+				{
+					const FReferenceSkeleton& RefSkeleton = Skeleton->GetReferenceSkeleton();
+					const int32 ParentBoneIndex = RefSkeleton.FindBoneIndex(BoneName);
+			
+					TArray<FName> ChildrenExcludedDescendents;
+					
+					if (bExcludeRecursively)
+					{
+						// create full list of child exclude with descendent search
+						for (const FName& ChildExcludedTrackName : ChildrenExcluded)
+						{
+							const int32 ChildExcludedIndex = RefSkeleton.FindBoneIndex(ChildExcludedTrackName);
+							for (const FName& ChildCandidateTrackName : TrackNames)
+							{
+								const int32 ChildCandidateTrackIndex = RefSkeleton.FindBoneIndex(ChildCandidateTrackName);
+								if (RefSkeleton.BoneIsChildOf(ChildCandidateTrackIndex, ChildExcludedIndex))
+								{
+									ChildrenExcludedDescendents.Add(ChildCandidateTrackName);
+								}
+							}
+						} 
+					}
+					
+					// slow
+					for (const FName& TrackName : TrackNames)
+					{
+						if (TrackName != BoneName && !ChildrenExcluded.Contains(TrackName) && !ChildrenExcludedDescendents.Contains(TrackName))
+						{
+							const int32 ChildBoneIndex = RefSkeleton.FindBoneIndex(TrackName);
+							if (RefSkeleton.BoneIsChildOf(ChildBoneIndex, ParentBoneIndex))
+							{
+								TracksToRemove.Add(TrackName);
+							}
+						}
+					}
+				}
+			}
+
+			IAnimationDataController& Controller = AnimationSequence->GetController();
+
+			IAnimationDataController::FScopedBracket ScopedBracket(Controller, LOCTEXT("RemoveBoneAnimation_Bracket", "Removing Bone Animation Track"));
+			for (const FName& TrackName : TracksToRemove)
+			{
+				Controller.RemoveBoneTrack(TrackName);
+			}
+		}
+		else
+		{
+			// print warning with track index
+			UE_LOG(LogAnimationBlueprintLibrary, Warning, TEXT("Invalid Bone Name for the animation."));
+		}
+	}
+}
+
 void UAnimationBlueprintLibrary::RemoveAllBoneAnimation(UAnimSequence* AnimationSequence)
 {
 	if (AnimationSequence)
diff --git a/Engine/Source/Editor/AnimationBlueprintLibrary/Public/AnimationBlueprintLibrary.h b/Engine/Source/Editor/AnimationBlueprintLibrary/Public/AnimationBlueprintLibrary.h
--- a/Engine/Source/Editor/AnimationBlueprintLibrary/Public/AnimationBlueprintLibrary.h
+++ b/Engine/Source/Editor/AnimationBlueprintLibrary/Public/AnimationBlueprintLibrary.h
@@ -403,6 +403,18 @@ public:
 	UFUNCTION(BlueprintCallable, Category = "AnimationBlueprintLibrary|Bones")
 	static void RemoveBoneAnimation(UAnimSequence* AnimationSequence, FName BoneName, bool bIncludeChildren = true, bool bFinalize = true);
 
+	/** Removes an Animation Curve by Name from the given Animation Sequence (Raw Animation Curves [Names] may not be removed from the Skeleton) 
+	 * 
+	 *	@param AnimationSequence : AnimSequence
+	 *	@param BoneName : Name of bone track user wants to remove
+	 *	@param bIncludeChildren : true if user wants to include all children of BoneName
+	 *  @param ChildrenExcluded : If bIncludeChildren is true, this is an optional list of child bones you want to exclude.
+	 *  @param bExcludeRecursively : If bChildrenExcluded is not empty, this will also exclude the descendants of bones in that list.
+	 *  @param bFinalize : If you set this to true, it will trigger compression. If you set bFinalize to be false, you'll have to manually trigger Finalize.
+	 */
+	UFUNCTION(BlueprintCallable, Category = "AnimationBlueprintLibrary|Bones", meta=(AutoCreateRefTerm="ChildrenExcluded"))
+	static void RemoveBoneSelectiveAnimation(UAnimSequence* AnimationSequence, FName BoneName, TArray<FName> ChildrenExcluded, bool bIncludeChildren = true, bool bExcludeRecursively = false, bool bFinalize = true);
+
 	/** Removes all Animation Bone Track Data from the given Animation Sequence */
 	UFUNCTION(BlueprintCallable, Category = "AnimationBlueprintLibrary|Curves")
 	static void RemoveAllBoneAnimation(UAnimSequence* AnimationSequence);
