RemoveBoneSelectiveAnimation
-changed Warning log calls to Display verbosity to avoid popup spam
-added invalid index check forchildren excluded list items

#rnx

diff --git a/Engine/Source/Editor/AnimationBlueprintLibrary/Private/AnimationBlueprintLibrary.cpp b/Engine/Source/Editor/AnimationBlueprintLibrary/Private/AnimationBlueprintLibrary.cpp
--- a/Engine/Source/Editor/AnimationBlueprintLibrary/Private/AnimationBlueprintLibrary.cpp
+++ b/Engine/Source/Editor/AnimationBlueprintLibrary/Private/AnimationBlueprintLibrary.cpp
@@ -2472,14 +2472,22 @@ void UAnimationBlueprintLibrary::RemoveBoneSelectiveAnimation(UAnimSequence* Ani
 						for (const FName& ChildExcludedTrackName : ChildrenExcluded)
 						{
 							const int32 ChildExcludedIndex = RefSkeleton.FindBoneIndex(ChildExcludedTrackName);
-							for (const FName& ChildCandidateTrackName : TrackNames)
+							if (RefSkeleton.IsValidIndex(ChildExcludedIndex))
 							{
-								const int32 ChildCandidateTrackIndex = RefSkeleton.FindBoneIndex(ChildCandidateTrackName);
-								if (RefSkeleton.BoneIsChildOf(ChildCandidateTrackIndex, ChildExcludedIndex))
+								for (const FName& ChildCandidateTrackName : TrackNames)
 								{
-									ChildrenExcludedDescendents.Add(ChildCandidateTrackName);
+									const int32 ChildCandidateTrackIndex = RefSkeleton.FindBoneIndex(ChildCandidateTrackName);
+									if (RefSkeleton.BoneIsChildOf(ChildCandidateTrackIndex, ChildExcludedIndex))
+									{
+										ChildrenExcludedDescendents.Add(ChildCandidateTrackName);
+									}
 								}
 							}
+							else
+							{
+								// print info with track index
+								UE_LOG(LogAnimationBlueprintLibrary, Display, TEXT("Invalid Child Exclusion Bone:"));
+							}
 						} 
 					}
 					
@@ -2508,8 +2516,8 @@ void UAnimationBlueprintLibrary::RemoveBoneSelectiveAnimation(UAnimSequence* Ani
 		}
 		else
 		{
-			// print warning with track index
-			UE_LOG(LogAnimationBlueprintLibrary, Warning, TEXT("Invalid Bone Name for the animation."));
+			// print info with track index
+			UE_LOG(LogAnimationBlueprintLibrary, Display, TEXT("Invalid Bone Name for the animation."));
 		}
 	}
 }
