Updated headers using UnrealCodeFixup to make sure dllstorage is on methods/staticvars instead of types. Part 4/n

diff --git a/Engine/Source/Editor/AnimGraph/Public/K2Node_AnimNodeReference.h b/Engine/Source/Editor/AnimGraph/Public/K2Node_AnimNodeReference.h
--- a/Engine/Source/Editor/AnimGraph/Public/K2Node_AnimNodeReference.h
+++ b/Engine/Source/Editor/AnimGraph/Public/K2Node_AnimNodeReference.h
@@ -15,19 +15,21 @@
 
 #include "K2Node_AnimNodeReference.generated.h"
 
+#define UE_API ANIMGRAPH_API
+
 class FBlueprintActionDatabaseRegistrar;
 class UEdGraph;
 class UObject;
 struct FSearchTagDataPair;
 
-UCLASS()
-class ANIMGRAPH_API UK2Node_AnimNodeReference : public UK2Node
+UCLASS(MinimalAPI)
+class UK2Node_AnimNodeReference : public UK2Node
 {
 public:
 	GENERATED_BODY()
 
 	// Get the text used for the node's label
-	FText GetLabelText() const;
+	UE_API FText GetLabelText() const;
 	
 	// Get the tag for this node, if any
 	FName GetTag() const { return Tag; }
@@ -40,17 +42,19 @@ private:
 	FName Tag;
 
 	// UEdGraphNode interface
-	virtual void AllocateDefaultPins() override;
-	virtual FText GetNodeTitle(ENodeTitleType::Type TitleType) const override;
-	virtual void AddSearchMetaDataInfo(TArray<FSearchTagDataPair>& OutTaggedMetaData) const override;
-	virtual bool IsCompatibleWithGraph(UEdGraph const* TargetGraph) const override;
+	UE_API virtual void AllocateDefaultPins() override;
+	UE_API virtual FText GetNodeTitle(ENodeTitleType::Type TitleType) const override;
+	UE_API virtual void AddSearchMetaDataInfo(TArray<FSearchTagDataPair>& OutTaggedMetaData) const override;
+	UE_API virtual bool IsCompatibleWithGraph(UEdGraph const* TargetGraph) const override;
 	
 	// UK2Node interface
-	virtual void ExpandNode(class FKismetCompilerContext& CompilerContext, UEdGraph* SourceGraph) override;
+	UE_API virtual void ExpandNode(class FKismetCompilerContext& CompilerContext, UEdGraph* SourceGraph) override;
 	virtual bool IsNodePure() const override { return true; }
-	virtual void GetMenuActions(FBlueprintActionDatabaseRegistrar& ActionRegistrar) const override;
-	virtual bool IsActionFilteredOut(class FBlueprintActionFilter const& Filter) override;
-	virtual FText GetMenuCategory() const override;
+	UE_API virtual void GetMenuActions(FBlueprintActionDatabaseRegistrar& ActionRegistrar) const override;
+	UE_API virtual bool IsActionFilteredOut(class FBlueprintActionFilter const& Filter) override;
+	UE_API virtual FText GetMenuCategory() const override;
 	virtual bool DrawNodeAsVariable() const override { return true; }
-	virtual FText GetTooltipText() const override;
-};
\ No newline at end of file
+	UE_API virtual FText GetTooltipText() const override;
+};
+
+#undef UE_API
diff --git a/Engine/Source/Editor/AnimationBlueprintEditor/Private/AnimationGraphFactory.h b/Engine/Source/Editor/AnimationBlueprintEditor/Private/AnimationGraphFactory.h
--- a/Engine/Source/Editor/AnimationBlueprintEditor/Private/AnimationGraphFactory.h
+++ b/Engine/Source/Editor/AnimationBlueprintEditor/Private/AnimationGraphFactory.h
@@ -6,19 +6,23 @@
 #include "HAL/Platform.h"
 #include "Templates/SharedPointer.h"
 
-struct ANIMATIONBLUEPRINTEDITOR_API FAnimationGraphNodeFactory : public FGraphPanelNodeFactory
+#define UE_API ANIMATIONBLUEPRINTEDITOR_API
+
+struct FAnimationGraphNodeFactory : public FGraphPanelNodeFactory
 {
-	virtual TSharedPtr<class SGraphNode> CreateNode(class UEdGraphNode* InNode) const override;
+	UE_API virtual TSharedPtr<class SGraphNode> CreateNode(class UEdGraphNode* InNode) const override;
 };
 
-struct ANIMATIONBLUEPRINTEDITOR_API FAnimationGraphPinFactory : public FGraphPanelPinFactory
+struct FAnimationGraphPinFactory : public FGraphPanelPinFactory
 {
 public:
-	virtual TSharedPtr<class SGraphPin> CreatePin(class UEdGraphPin* Pin) const override;
+	UE_API virtual TSharedPtr<class SGraphPin> CreatePin(class UEdGraphPin* Pin) const override;
 };
 
-struct ANIMATIONBLUEPRINTEDITOR_API FAnimationGraphPinConnectionFactory : public FGraphPanelPinConnectionFactory
+struct FAnimationGraphPinConnectionFactory : public FGraphPanelPinConnectionFactory
 {
 public:
-	virtual class FConnectionDrawingPolicy* CreateConnectionPolicy(const class UEdGraphSchema* Schema, int32 InBackLayerID, int32 InFrontLayerID, float ZoomFactor, const class FSlateRect& InClippingRect, class FSlateWindowElementList& InDrawElements, class UEdGraph* InGraphObj) const override;
-};
\ No newline at end of file
+	UE_API virtual class FConnectionDrawingPolicy* CreateConnectionPolicy(const class UEdGraphSchema* Schema, int32 InBackLayerID, int32 InFrontLayerID, float ZoomFactor, const class FSlateRect& InClippingRect, class FSlateWindowElementList& InDrawElements, class UEdGraph* InGraphObj) const override;
+};
+
+#undef UE_API
diff --git a/Engine/Source/Editor/AnimationBlueprintEditor/Public/AnimationNodes/SAnimationGraphNode.h b/Engine/Source/Editor/AnimationBlueprintEditor/Public/AnimationNodes/SAnimationGraphNode.h
--- a/Engine/Source/Editor/AnimationBlueprintEditor/Public/AnimationNodes/SAnimationGraphNode.h
+++ b/Engine/Source/Editor/AnimationBlueprintEditor/Public/AnimationNodes/SAnimationGraphNode.h
@@ -15,6 +15,8 @@
 #include "Templates/SharedPointer.h"
 #include "Widgets/DeclarativeSyntaxSupport.h"
 
+#define UE_API ANIMATIONBLUEPRINTEDITOR_API
+
 class IDetailTreeNode;
 class IPropertyRowGenerator;
 class SGraphNode;
@@ -33,7 +35,7 @@ struct FNodeInfoContext;
 struct FOverlayWidgetInfo;
 template <typename FuncType> class TFunctionRef;
 
-class ANIMATIONBLUEPRINTEDITOR_API SAnimationGraphNode : public SGraphNodeK2Base
+class SAnimationGraphNode : public SGraphNodeK2Base
 {
 public:
 	SLATE_BEGIN_ARGS(SAnimationGraphNode) {}
@@ -45,41 +47,41 @@ public:
 	// Reverse index of the tag/functions slot
 	static const int32 TagAndFunctionsSlotReverseIndex = 1;
 	
-	void Construct(const FArguments& InArgs, UAnimGraphNode_Base* InNode);
+	UE_API void Construct(const FArguments& InArgs, UAnimGraphNode_Base* InNode);
 
 	// Tweak any created pin widgets so they respond to bindings
-	static void ReconfigurePinWidgetsForPropertyBindings(UAnimGraphNode_Base* InAnimGraphNode, TSharedRef<SGraphNode> InGraphNodeWidget, TFunctionRef<TSharedPtr<SGraphPin>(UEdGraphPin*)> InFindWidgetForPin);
+	static UE_API void ReconfigurePinWidgetsForPropertyBindings(UAnimGraphNode_Base* InAnimGraphNode, TSharedRef<SGraphNode> InGraphNodeWidget, TFunctionRef<TSharedPtr<SGraphPin>(UEdGraphPin*)> InFindWidgetForPin);
 
 	// Create below-widget controls for editing anim node functions
-	static TSharedRef<SWidget> CreateNodeFunctionsWidget(UAnimGraphNode_Base* InAnimNode, TAttribute<bool> InUseLowDetail);
+	static UE_API TSharedRef<SWidget> CreateNodeFunctionsWidget(UAnimGraphNode_Base* InAnimNode, TAttribute<bool> InUseLowDetail);
 
 	// Create below-widget controls for editing anim node tags
-	static TSharedRef<SWidget> CreateNodeTagWidget(UAnimGraphNode_Base* InAnimNode, TAttribute<bool> InUseLowDetail);
+	static UE_API TSharedRef<SWidget> CreateNodeTagWidget(UAnimGraphNode_Base* InAnimNode, TAttribute<bool> InUseLowDetail);
 	
 protected:
 	// SWidget interface
-	virtual void Tick(const FGeometry& AllottedGeometry, const double InCurrentTime, const float InDeltaTime) override;
+	UE_API virtual void Tick(const FGeometry& AllottedGeometry, const double InCurrentTime, const float InDeltaTime) override;
 	// End of SWidget interface
 
 	// SGraphNode interface
-	virtual TArray<FOverlayWidgetInfo> GetOverlayWidgets(bool bSelected, const FVector2f& WidgetSize) const override;
-	virtual TSharedRef<SWidget> CreateTitleWidget(TSharedPtr<SNodeTitle> InNodeTitle) override;
-	virtual void GetNodeInfoPopups(FNodeInfoContext* Context, TArray<FGraphInformationPopupInfo>& Popups) const override;
-	virtual void CreateBelowPinControls(TSharedPtr<SVerticalBox> MainBox) override;
-	virtual TSharedRef<SWidget> CreateNodeContentArea() override;
+	UE_API virtual TArray<FOverlayWidgetInfo> GetOverlayWidgets(bool bSelected, const FVector2f& WidgetSize) const override;
+	UE_API virtual TSharedRef<SWidget> CreateTitleWidget(TSharedPtr<SNodeTitle> InNodeTitle) override;
+	UE_API virtual void GetNodeInfoPopups(FNodeInfoContext* Context, TArray<FGraphInformationPopupInfo>& Popups) const override;
+	UE_API virtual void CreateBelowPinControls(TSharedPtr<SVerticalBox> MainBox) override;
+	UE_API virtual TSharedRef<SWidget> CreateNodeContentArea() override;
 	virtual bool IsHidingPinWidgets() const override { return UseLowDetailNodeContent(); }
 	// End of SGraphNode interface
 
 private:
 	// Handle the node informing us that the title has changed
-	void HandleNodeTitleChanged();
+	UE_API void HandleNodeTitleChanged();
 
 	// LOD related functions for content area
-	bool UseLowDetailNodeContent() const;
-	FVector2D GetLowDetailDesiredSize() const;
+	UE_API bool UseLowDetailNodeContent() const;
+	UE_API FVector2D GetLowDetailDesiredSize() const;
 
 	// Handler for pose watches changing
-	void HandlePoseWatchesChanged(UAnimBlueprint* InAnimBlueprint, UEdGraphNode* InNode);
+	UE_API void HandlePoseWatchesChanged(UAnimBlueprint* InAnimBlueprint, UEdGraphNode* InNode);
 
 	/** Keep a reference to the indicator widget handing around */
 	TSharedPtr<SWidget> IndicatorWidget;
@@ -96,3 +98,5 @@ private:
 	/** Cached content area widget (used to derive LastHighDetailSize) */
 	TSharedPtr<SWidget> CachedContentArea;
 };
+
+#undef UE_API
diff --git a/Engine/Source/Editor/AnimationBlueprintLibrary/Public/AnimPose.h b/Engine/Source/Editor/AnimationBlueprintLibrary/Public/AnimPose.h
--- a/Engine/Source/Editor/AnimationBlueprintLibrary/Public/AnimPose.h
+++ b/Engine/Source/Editor/AnimationBlueprintLibrary/Public/AnimPose.h
@@ -16,6 +16,8 @@
 
 #include "AnimPose.generated.h"
 
+#define UE_API ANIMATIONBLUEPRINTLIBRARY_API
+
 class UAnimBlueprint;
 class UAnimSequenceBase;
 class UObject;
@@ -47,7 +49,7 @@ enum class EAnimDataEvalType : uint8
 };
 
 USTRUCT(BlueprintType)
-struct ANIMATIONBLUEPRINTLIBRARY_API FAnimPoseEvaluationOptions
+struct FAnimPoseEvaluationOptions
 {
 	GENERATED_BODY()
 
@@ -82,29 +84,29 @@ struct ANIMATIONBLUEPRINTLIBRARY_API FAnimPoseEvaluationOptions
 
 /** Script friendly representation of an evaluated animation bone pose */
 USTRUCT(BlueprintType)
-struct ANIMATIONBLUEPRINTLIBRARY_API FAnimPose
+struct FAnimPose
 {
 	GENERATED_BODY()
 
 	/** Returns whether or not the pose data was correctly initialized and populated */
-	bool IsValid() const;
+	UE_API bool IsValid() const;
 	
 protected:
 	/** Initializes the various arrays, using and copying the provided bone container */
-	void Init(const FBoneContainer& InBoneContainer);
+	UE_API void Init(const FBoneContainer& InBoneContainer);
 	/** Populates an FCompactPose using the contained bone data */
-	void GetPose(FCompactPose& InOutCompactPose) const;
+	UE_API void GetPose(FCompactPose& InOutCompactPose) const;
 	/** Generates the contained bone data using the provided Component and its AnimInstance */
-	void SetPose(USkeletalMeshComponent* Component);
+	UE_API void SetPose(USkeletalMeshComponent* Component);
 	/** Generates the contained bone data using the provided CompactPose */
-	void SetPose(const FAnimationPoseData& PoseData);
+	UE_API void SetPose(const FAnimationPoseData& PoseData);
 	/** Copies the reference pose to animated pose data */
-	void SetToRefPose();
+	UE_API void SetToRefPose();
 
 	/** (Re-)Generates the world space transforms using populated local space data */
-	void GenerateWorldSpaceTransforms();
+	UE_API void GenerateWorldSpaceTransforms();
 	/** Resets all contained data, rendering the instance invalid */
-	void Reset();
+	UE_API void Reset();
 
 	/** Whether or not the contained data was initialized and can be used to store a pose */
 	bool IsInitialized() const { return BoneNames.Num() != 0; }
@@ -152,8 +154,8 @@ protected:
 };
 
 /** Script exposed functionality for populating, retrieving data from and setting data on FAnimPose */
-UCLASS()
-class ANIMATIONBLUEPRINTLIBRARY_API UAnimPoseExtensions : public UBlueprintFunctionLibrary
+UCLASS(MinimalAPI)
+class UAnimPoseExtensions : public UBlueprintFunctionLibrary
 {
 	GENERATED_BODY()
 	
@@ -166,7 +168,7 @@ public:
 	* @return	Result of the validation
 	*/	
 	UFUNCTION(BlueprintPure, meta = (ScriptMethod), Category = "Animation|Pose")
-	static bool IsValid(UPARAM(ref) const FAnimPose& Pose);
+	static UE_API bool IsValid(UPARAM(ref) const FAnimPose& Pose);
 
 	/**
 	* Returns an array of bone names contained by the pose
@@ -176,7 +178,7 @@ public:
 	*
 	*/
 	UFUNCTION(BlueprintPure, meta = (ScriptMethod), Category = "Animation|Pose")
-	static void GetBoneNames(UPARAM(ref) const FAnimPose& Pose, TArray<FName>& Bones);
+	static UE_API void GetBoneNames(UPARAM(ref) const FAnimPose& Pose, TArray<FName>& Bones);
 
 	/**
 	* Retrieves the transform for the provided bone name from a pose
@@ -188,7 +190,7 @@ public:
 	* @return	Transform in requested space for bone if found, otherwise return identity transform
 	*/
 	UFUNCTION(BlueprintPure, meta = (ScriptMethod), Category = "Animation|Pose")
-    static const FTransform& GetBonePose(UPARAM(ref) const FAnimPose& Pose, FName BoneName, EAnimPoseSpaces Space = EAnimPoseSpaces::Local);
+    static UE_API const FTransform& GetBonePose(UPARAM(ref) const FAnimPose& Pose, FName BoneName, EAnimPoseSpaces Space = EAnimPoseSpaces::Local);
 	
 	/**
 	* Sets the transform for the provided bone name for a pose
@@ -199,7 +201,7 @@ public:
 	* @param	Space		Space in which the transform should be set
 	*/
 	UFUNCTION(BlueprintCallable, meta = (ScriptMethod), Category = "Animation|Pose")
-	static void SetBonePose(UPARAM(ref) FAnimPose& Pose, FTransform Transform, FName BoneName, EAnimPoseSpaces Space = EAnimPoseSpaces::Local);	
+	static UE_API void SetBonePose(UPARAM(ref) FAnimPose& Pose, FTransform Transform, FName BoneName, EAnimPoseSpaces Space = EAnimPoseSpaces::Local);	
 
 	/**
 	* Retrieves the reference pose transform for the provided bone name
@@ -211,7 +213,7 @@ public:
 	* @return	Transform in requested space for bone if found, otherwise return identity transform
 	*/
 	UFUNCTION(BlueprintPure, meta = (ScriptMethod), Category = "Animation|Pose")
-    static const FTransform& GetRefBonePose(UPARAM(ref) const FAnimPose& Pose, FName BoneName, EAnimPoseSpaces Space = EAnimPoseSpaces::Local);
+    static UE_API const FTransform& GetRefBonePose(UPARAM(ref) const FAnimPose& Pose, FName BoneName, EAnimPoseSpaces Space = EAnimPoseSpaces::Local);
 
 	/**
 	* Retrieves the relative transform between the two provided bone names
@@ -224,7 +226,7 @@ public:
 	* @return	Relative transform in requested space for bone if found, otherwise return identity transform
 	*/
 	UFUNCTION(BlueprintPure, meta = (ScriptMethod), Category = "Animation|Pose")
-    static FTransform GetRelativeTransform(UPARAM(ref) const FAnimPose& Pose, FName FromBoneName, FName ToBoneName, EAnimPoseSpaces Space = EAnimPoseSpaces::Local);
+    static UE_API FTransform GetRelativeTransform(UPARAM(ref) const FAnimPose& Pose, FName FromBoneName, FName ToBoneName, EAnimPoseSpaces Space = EAnimPoseSpaces::Local);
 
 	/**
 	* Retrieves the relative transform between reference and animated bone transform
@@ -236,7 +238,7 @@ public:
 	* @return	Relative transform in requested space for bone if found, otherwise return identity transform
 	*/
 	UFUNCTION(BlueprintPure, meta = (ScriptMethod), Category = "Animation|Pose")
-    static FTransform GetRelativeToRefPoseTransform(UPARAM(ref) const FAnimPose& Pose, FName BoneName, EAnimPoseSpaces Space = EAnimPoseSpaces::Local);
+    static UE_API FTransform GetRelativeToRefPoseTransform(UPARAM(ref) const FAnimPose& Pose, FName BoneName, EAnimPoseSpaces Space = EAnimPoseSpaces::Local);
 
 	/**
 	* Retrieves the relative transform for the reference pose between the two provided bone names 
@@ -249,7 +251,7 @@ public:
 	* @return	Relative transform in requested space for bone if found, otherwise return identity transform
 	*/
 	UFUNCTION(BlueprintPure, meta = (ScriptMethod), Category = "Animation|Pose")
-    static FTransform GetRefPoseRelativeTransform(UPARAM(ref) const FAnimPose& Pose, FName FromBoneName, FName ToBoneName, EAnimPoseSpaces Space = EAnimPoseSpaces::Local);
+    static UE_API FTransform GetRefPoseRelativeTransform(UPARAM(ref) const FAnimPose& Pose, FName FromBoneName, FName ToBoneName, EAnimPoseSpaces Space = EAnimPoseSpaces::Local);
 
 	/**
 	* Returns an array of socket names contained by the pose
@@ -259,7 +261,7 @@ public:
 	*
 	*/
 	UFUNCTION(BlueprintPure, meta = (ScriptMethod), Category = "Animation|Pose")
-	static void GetSocketNames(UPARAM(ref) const FAnimPose& Pose, TArray<FName>& Sockets);
+	static UE_API void GetSocketNames(UPARAM(ref) const FAnimPose& Pose, TArray<FName>& Sockets);
 
 	/**
 	* Retrieves the transform for the provided socket name from a pose
@@ -271,7 +273,7 @@ public:
 	* @return	Transform in requested space for bone if found, otherwise return identity transform
 	*/
 	UFUNCTION(BlueprintPure, meta = (ScriptMethod), Category = "Animation|Pose")
-	static FTransform GetSocketPose(UPARAM(ref) const FAnimPose& Pose, FName SocketName, EAnimPoseSpaces Space = EAnimPoseSpaces::Local);
+	static UE_API FTransform GetSocketPose(UPARAM(ref) const FAnimPose& Pose, FName SocketName, EAnimPoseSpaces Space = EAnimPoseSpaces::Local);
 	
 
 	/**
@@ -283,7 +285,7 @@ public:
 	* @param	Pose					Anim Pose to hold the evaluated data
 	*/
 	UFUNCTION(BlueprintCallable, meta = (ScriptMethod), Category = "Animation|Pose")
-	static void GetAnimPoseAtTime(const UAnimSequenceBase* AnimationSequenceBase, double Time, FAnimPoseEvaluationOptions EvaluationOptions, FAnimPose& Pose);
+	static UE_API void GetAnimPoseAtTime(const UAnimSequenceBase* AnimationSequenceBase, double Time, FAnimPoseEvaluationOptions EvaluationOptions, FAnimPose& Pose);
 
 	/**
 	* Evaluates an Animation Sequence Base at different time intervals to generate a valid Anim Pose instances
@@ -293,7 +295,7 @@ public:
 	* @param	EvaluationOptions		Options determining the way the pose should be evaluated
 	* @param	InOutPoses				Anim Poses holding the evaluated data (number matches TimeIntervals)
 	*/
-	static void GetAnimPoseAtTimeIntervals(const UAnimSequenceBase* AnimationSequenceBase, TArray<double> TimeIntervals, FAnimPoseEvaluationOptions EvaluationOptions, TArray<FAnimPose>& InOutPoses);
+	static UE_API void GetAnimPoseAtTimeIntervals(const UAnimSequenceBase* AnimationSequenceBase, TArray<double> TimeIntervals, FAnimPoseEvaluationOptions EvaluationOptions, TArray<FAnimPose>& InOutPoses);
 
 	/**
 	* Evaluates an Animation Sequence Base to generate a valid Anim Pose instance
@@ -304,7 +306,7 @@ public:
 	* @param	Pose					Anim Pose to hold the evaluated data
 	*/
 	UFUNCTION(BlueprintCallable, meta = (ScriptMethod), Category = "Animation|Pose")
-	static void GetAnimPoseAtFrame(const UAnimSequenceBase* AnimationSequenceBase, int32 FrameIndex, FAnimPoseEvaluationOptions EvaluationOptions, FAnimPose& Pose);
+	static UE_API void GetAnimPoseAtFrame(const UAnimSequenceBase* AnimationSequenceBase, int32 FrameIndex, FAnimPoseEvaluationOptions EvaluationOptions, FAnimPose& Pose);
 
 	/**
 	* Evaluates an Animation Blueprint instance, using the provided Anim Pose and its Input Pose value, generating a valid Anim Pose using the result. Warning this function may cause performance issues.
@@ -315,7 +317,7 @@ public:
 	* @param	OutPose					Anim pose to hold the data from evaluating the Animation Blueprint instance		
 	*/
 	UFUNCTION(BlueprintCallable, meta = (ScriptMethod), Category = "Animation|Pose")
-	static void EvaluateAnimationBlueprintWithInputPose(const FAnimPose& InputPose, USkeletalMesh* TargetSkeletalMesh, UAnimBlueprint* AnimationBlueprint, FAnimPose& OutPose);
+	static UE_API void EvaluateAnimationBlueprintWithInputPose(const FAnimPose& InputPose, USkeletalMesh* TargetSkeletalMesh, UAnimBlueprint* AnimationBlueprint, FAnimPose& OutPose);
 
 	/**
 	* Populates an Anim Pose with the reference poses stored for the provided USkeleton
@@ -324,7 +326,7 @@ public:
 	* @param	OutPose					Anim pose to hold the reference pose
 	*/
 	UFUNCTION(BlueprintPure, meta = (ScriptMethod), Category = "Animation|Pose")
-	static void GetReferencePose(USkeleton* Skeleton, FAnimPose& OutPose);
+	static UE_API void GetReferencePose(USkeleton* Skeleton, FAnimPose& OutPose);
 
 	/**
 	* Returns an array of curve names contained by the pose
@@ -334,7 +336,7 @@ public:
 	*
 	*/
 	UFUNCTION(BlueprintPure, meta = (ScriptMethod), Category = "Animation|Pose")
-	static void GetCurveNames(UPARAM(ref) const FAnimPose& Pose, TArray<FName>& Curves);
+	static UE_API void GetCurveNames(UPARAM(ref) const FAnimPose& Pose, TArray<FName>& Curves);
 	
 	/**
     * Returns the weight of an evaluated curve - if found
@@ -345,7 +347,7 @@ public:
     * @return	Curve weight value, if found - 0.f otherwise
     */
     UFUNCTION(BlueprintPure, meta = (ScriptMethod), Category = "Animation|Pose")
-    static float GetCurveWeight(UPARAM(ref) const FAnimPose& Pose, const FName& CurveName);
+    static UE_API float GetCurveWeight(UPARAM(ref) const FAnimPose& Pose, const FName& CurveName);
 
 
 	/**
@@ -354,5 +356,7 @@ public:
     * @param	Pose		Anim Pose to retrieve the value from
     * @param	InOutCompactPose	CompactPose from the Pose
     */
-	static void GetCompactPose(const FAnimPose& Pose, FCompactPose& InOutCompactPose);
-};
\ No newline at end of file
+	static UE_API void GetCompactPose(const FAnimPose& Pose, FCompactPose& InOutCompactPose);
+};
+
+#undef UE_API
diff --git a/Engine/Source/Editor/AnimationBlueprintLibrary/Public/AnimationAttributeBlueprintLibrary.h b/Engine/Source/Editor/AnimationBlueprintLibrary/Public/AnimationAttributeBlueprintLibrary.h
--- a/Engine/Source/Editor/AnimationBlueprintLibrary/Public/AnimationAttributeBlueprintLibrary.h
+++ b/Engine/Source/Editor/AnimationBlueprintLibrary/Public/AnimationAttributeBlueprintLibrary.h
@@ -13,6 +13,8 @@
 
 #include "AnimationAttributeBlueprintLibrary.generated.h"
 
+#define UE_API ANIMATIONBLUEPRINTLIBRARY_API
+
 class FArrayProperty;
 class IAnimationDataController;
 class UAnimDataModel;
@@ -21,30 +23,32 @@ class UScriptStruct;
 struct FAnimationAttributeIdentifier;
 struct FFrame;
 
-UCLASS(meta=(ScriptName="AnimationAttributeLibrary"))
-class ANIMATIONBLUEPRINTLIBRARY_API UAnimationAttributeBlueprintLibrary : public UBlueprintFunctionLibrary
+UCLASS(MinimalAPI, meta=(ScriptName="AnimationAttributeLibrary"))
+class UAnimationAttributeBlueprintLibrary : public UBlueprintFunctionLibrary
 {
 	GENERATED_BODY()
 
 public:
 	UFUNCTION(BlueprintCallable, CustomThunk, Category = "Animation|Attributes", meta=(CustomStructureParam="Value", BlueprintInternalUseOnly="true"))
-	static bool SetAttributeKey(TScriptInterface<IAnimationDataController> AnimationDataController, const FAnimationAttributeIdentifier& AttributeIdentifier, float Time, const int32& Value);
-	static bool Generic_SetAttributeKey(TScriptInterface<IAnimationDataController> AnimationDataController, const FAnimationAttributeIdentifier& AttributeIdentifier, float Time, UScriptStruct* ScriptStruct, const void* Value);
+	static UE_API bool SetAttributeKey(TScriptInterface<IAnimationDataController> AnimationDataController, const FAnimationAttributeIdentifier& AttributeIdentifier, float Time, const int32& Value);
+	static UE_API bool Generic_SetAttributeKey(TScriptInterface<IAnimationDataController> AnimationDataController, const FAnimationAttributeIdentifier& AttributeIdentifier, float Time, UScriptStruct* ScriptStruct, const void* Value);
 	DECLARE_FUNCTION(execSetAttributeKey);
 
 	UFUNCTION(BlueprintCallable, CustomThunk, Category = "Animation|Attributes", meta=(CustomStructureParam="Values", BlueprintInternalUseOnly="true"))
-	static bool SetAttributeKeys(TScriptInterface<IAnimationDataController> AnimationDataController, const FAnimationAttributeIdentifier& AttributeIdentifier, const TArray<float>& Times, const TArray<int32>& Values);
-	static bool Generic_SetAttributeKeys(TScriptInterface<IAnimationDataController> AnimationDataController, const FAnimationAttributeIdentifier& AttributeIdentifier, const TArray<float>& Times, const void* ValuesArray, const FArrayProperty* ValuesArrayProperty);
+	static UE_API bool SetAttributeKeys(TScriptInterface<IAnimationDataController> AnimationDataController, const FAnimationAttributeIdentifier& AttributeIdentifier, const TArray<float>& Times, const TArray<int32>& Values);
+	static UE_API bool Generic_SetAttributeKeys(TScriptInterface<IAnimationDataController> AnimationDataController, const FAnimationAttributeIdentifier& AttributeIdentifier, const TArray<float>& Times, const void* ValuesArray, const FArrayProperty* ValuesArrayProperty);
 	DECLARE_FUNCTION(execSetAttributeKeys);
 
 
 	UFUNCTION(BlueprintCallable, CustomThunk, Category = "Animation|Attributes", meta=(CustomStructureParam="Value", BlueprintInternalUseOnly="true"))
-	static bool GetAttributeKey(TScriptInterface<IAnimationDataModel> AnimationDataModel, const FAnimationAttributeIdentifier& AttributeIdentifier, float Time, int32& Value);
-	static bool Generic_GetAttributeKey(TScriptInterface<IAnimationDataModel> AnimationDataModel, const FAnimationAttributeIdentifier& AttributeIdentifier, float Time, UScriptStruct* ScriptStruct, void* Value);
+	static UE_API bool GetAttributeKey(TScriptInterface<IAnimationDataModel> AnimationDataModel, const FAnimationAttributeIdentifier& AttributeIdentifier, float Time, int32& Value);
+	static UE_API bool Generic_GetAttributeKey(TScriptInterface<IAnimationDataModel> AnimationDataModel, const FAnimationAttributeIdentifier& AttributeIdentifier, float Time, UScriptStruct* ScriptStruct, void* Value);
 	DECLARE_FUNCTION(execGetAttributeKey);
 
 	UFUNCTION(BlueprintCallable, CustomThunk, Category = "Animation|Attributes", meta=(CustomStructureParam="Values", BlueprintInternalUseOnly="true"))
-	static bool GetAttributeKeys(TScriptInterface<IAnimationDataModel> AnimationDataModel, const FAnimationAttributeIdentifier& AttributeIdentifier, TArray<float>& OutTimes, TArray<int32>& Values);
-	static bool Generic_GetAttributeKeys(TScriptInterface<IAnimationDataModel> AnimationDataModel, const FAnimationAttributeIdentifier& AttributeIdentifier, TArray<float>& Times, void* ValuesArray, const FArrayProperty* ValuesArrayProperty);
+	static UE_API bool GetAttributeKeys(TScriptInterface<IAnimationDataModel> AnimationDataModel, const FAnimationAttributeIdentifier& AttributeIdentifier, TArray<float>& OutTimes, TArray<int32>& Values);
+	static UE_API bool Generic_GetAttributeKeys(TScriptInterface<IAnimationDataModel> AnimationDataModel, const FAnimationAttributeIdentifier& AttributeIdentifier, TArray<float>& Times, void* ValuesArray, const FArrayProperty* ValuesArrayProperty);
 	DECLARE_FUNCTION(execGetAttributeKeys);
 };
+
+#undef UE_API
diff --git a/Engine/Source/Editor/AnimationBlueprintLibrary/Public/AnimationBlueprintLibrary.h b/Engine/Source/Editor/AnimationBlueprintLibrary/Public/AnimationBlueprintLibrary.h
--- a/Engine/Source/Editor/AnimationBlueprintLibrary/Public/AnimationBlueprintLibrary.h
+++ b/Engine/Source/Editor/AnimationBlueprintLibrary/Public/AnimationBlueprintLibrary.h
@@ -24,6 +24,8 @@
 
 #include "AnimationBlueprintLibrary.generated.h"
 
+#define UE_API ANIMATIONBLUEPRINTLIBRARY_API
+
 class UAnimBlueprint;
 class UAnimBoneCompressionSettings;
 class UAnimCompress;
@@ -60,35 +62,35 @@ DECLARE_DYNAMIC_DELEGATE_TwoParams(FOnNotifyReplaced, UAnimNotify*, OldNotify, U
 DECLARE_DYNAMIC_DELEGATE_TwoParams(FOnNotifyStateReplaced, UAnimNotifyState*, OldNotifyState, UAnimNotifyState*, NewNotifyState);
 
 /** Blueprint library for altering and analyzing animation / skeletal data */
-UCLASS(meta=(ScriptName="AnimationLibrary"))
-class ANIMATIONBLUEPRINTLIBRARY_API UAnimationBlueprintLibrary : public UBlueprintFunctionLibrary
+UCLASS(MinimalAPI, meta=(ScriptName="AnimationLibrary"))
+class UAnimationBlueprintLibrary : public UBlueprintFunctionLibrary
 {
 	GENERATED_BODY()
 
 public:
 	/** Retrieves the number of animation frames for the given Animation Sequence */
 	UFUNCTION(BlueprintPure, meta=(AutoCreateRefTerm = "AnimationSequence"), Category = "AnimationBlueprintLibrary|Animation")
-	static void GetNumFrames(const UAnimSequenceBase* AnimationSequenceBase, int32& NumFrames);
+	static UE_API void GetNumFrames(const UAnimSequenceBase* AnimationSequenceBase, int32& NumFrames);
 
 	/** Retrieves the number of animation keys for the given Animation Sequence */
 	UFUNCTION(BlueprintPure, meta = (AutoCreateRefTerm = "AnimationSequence"), Category = "AnimationBlueprintLibrary|Animation")
-	static void GetNumKeys(const UAnimSequenceBase* AnimationSequenceBase, int32& NumKeys);
+	static UE_API void GetNumKeys(const UAnimSequenceBase* AnimationSequenceBase, int32& NumKeys);
 	
 	/** Retrieves the Names of the individual ATracks for the given Animation Sequence */
 	UFUNCTION(BlueprintPure, Category = "AnimationBlueprintLibrary|Animation")
-	static void GetAnimationTrackNames(const UAnimSequenceBase* AnimationSequenceBase, TArray<FName>& TrackNames);
+	static UE_API void GetAnimationTrackNames(const UAnimSequenceBase* AnimationSequenceBase, TArray<FName>& TrackNames);
 
 	/** Retrieves the Names of the Animation Slots used in the given Montage */
 	UFUNCTION(BlueprintPure, Category = "AnimationBlueprintLibrary|Montage")
-	static void GetMontageSlotNames(const UAnimMontage* AnimationMontage, TArray<FName>& SlotNames);
+	static UE_API void GetMontageSlotNames(const UAnimMontage* AnimationMontage, TArray<FName>& SlotNames);
 
 	/** Retrieves the Names of the individual float curves for the given Animation Sequence */
 	UFUNCTION(BlueprintPure, Category = "AnimationBlueprintLibrary|Animation")
-	static void GetAnimationCurveNames(const UAnimSequenceBase* AnimationSequenceBase, ERawCurveTrackTypes CurveType, TArray<FName>& CurveNames);
+	static UE_API void GetAnimationCurveNames(const UAnimSequenceBase* AnimationSequenceBase, ERawCurveTrackTypes CurveType, TArray<FName>& CurveNames);
 
 	/** Gets the root transform from the raw animation at Time */
 	UFUNCTION(BlueprintCallable, Category="AnimationBlueprintLibrary|Animation")
-	static FTransform ExtractRootTrackTransform(const UAnimSequenceBase* AnimationSequenceBase, float Time);
+	static UE_API FTransform ExtractRootTrackTransform(const UAnimSequenceBase* AnimationSequenceBase, float Time);
 
 	/** Retrieves the Raw Translation Animation Data for the given Animation Track Name and Animation Sequence */
 	UE_DEPRECATED(5.2, "GetRawTrackPositionData has been deprecated, use AnimationModel interface instead")
@@ -115,274 +117,274 @@ public:
 	static bool IsValidRawAnimationTrackName(const UAnimSequenceBase* AnimationSequenceBase, const FName TrackName) { return false; }
 
 	UE_DEPRECATED(5.2, "GetRawAnimationTrackByName has been deprecated")
-	static const FRawAnimSequenceTrack& GetRawAnimationTrackByName(const UAnimSequenceBase* AnimationSequenceBase, const FName TrackName);
+	static UE_API const FRawAnimSequenceTrack& GetRawAnimationTrackByName(const UAnimSequenceBase* AnimationSequenceBase, const FName TrackName);
 
 	// Compression
 
 	/** Retrieves the Bone Compression Settings for the given Animation Sequence */
 	UFUNCTION(BlueprintPure, Category = "AnimationBlueprintLibrary|Compression")
-	static void GetBoneCompressionSettings(const UAnimSequence* AnimationSequence, UAnimBoneCompressionSettings*& CompressionSettings);
+	static UE_API void GetBoneCompressionSettings(const UAnimSequence* AnimationSequence, UAnimBoneCompressionSettings*& CompressionSettings);
 
 	/** Sets the Bone Compression Settings for the given Animation Sequence */
 	UFUNCTION(BlueprintCallable, Category = "AnimationBlueprintLibrary|Compression")
-	static void SetBoneCompressionSettings(UAnimSequence* AnimationSequence, UAnimBoneCompressionSettings* CompressionSettings);
+	static UE_API void SetBoneCompressionSettings(UAnimSequence* AnimationSequence, UAnimBoneCompressionSettings* CompressionSettings);
 
 	/** Retrieves the Curve Compression Settings for the given Animation Sequence */
 	UFUNCTION(BlueprintPure, Category = "AnimationBlueprintLibrary|Compression")
-	static void GetCurveCompressionSettings(const UAnimSequence* AnimationSequence, UAnimCurveCompressionSettings*& CompressionSettings);
+	static UE_API void GetCurveCompressionSettings(const UAnimSequence* AnimationSequence, UAnimCurveCompressionSettings*& CompressionSettings);
 
 	/** Sets the Curve Compression Settings for the given Animation Sequence */
 	UFUNCTION(BlueprintCallable, Category = "AnimationBlueprintLibrary|Compression")
-	static void SetCurveCompressionSettings(UAnimSequence* AnimationSequence, UAnimCurveCompressionSettings* CompressionSettings);
+	static UE_API void SetCurveCompressionSettings(UAnimSequence* AnimationSequence, UAnimCurveCompressionSettings* CompressionSettings);
 
 	/** Retrieves the Variable Frame Stripping Settings for the given Animation Sequence */
 	UFUNCTION(BlueprintPure, Category = "AnimationBlueprintLibrary|Compression")
-	static void GetVariableFrameStrippingSettings(const UAnimSequence* AnimationSequence, UVariableFrameStrippingSettings*& VariableFrameStrippingSettings);
+	static UE_API void GetVariableFrameStrippingSettings(const UAnimSequence* AnimationSequence, UVariableFrameStrippingSettings*& VariableFrameStrippingSettings);
 
 	/** Sets the Variable Frame Stripping Settings for the given Animation Sequence */
 	UFUNCTION(BlueprintCallable, Category = "AnimationBlueprintLibrary|Compression")
-	static void SetVariableFrameStrippingSettings(UAnimSequence* AnimationSequence, UVariableFrameStrippingSettings* VariableFrameStrippingSettings);
+	static UE_API void SetVariableFrameStrippingSettings(UAnimSequence* AnimationSequence, UVariableFrameStrippingSettings* VariableFrameStrippingSettings);
 
 	// Additive 
 	/** Retrieves the Additive Animation type for the given Animation Sequence */
 	UFUNCTION(BlueprintPure, Category = "AnimationBlueprintLibrary|Additive")
-	static void GetAdditiveAnimationType(const UAnimSequence* AnimationSequence, TEnumAsByte<enum EAdditiveAnimationType>& AdditiveAnimationType);
+	static UE_API void GetAdditiveAnimationType(const UAnimSequence* AnimationSequence, TEnumAsByte<enum EAdditiveAnimationType>& AdditiveAnimationType);
 
 	/** Sets the Additive Animation type for the given Animation Sequence */
 	UFUNCTION(BlueprintCallable, Category = "AnimationBlueprintLibrary|Additive")
-	static void SetAdditiveAnimationType(UAnimSequence* AnimationSequence, const TEnumAsByte<enum EAdditiveAnimationType> AdditiveAnimationType);
+	static UE_API void SetAdditiveAnimationType(UAnimSequence* AnimationSequence, const TEnumAsByte<enum EAdditiveAnimationType> AdditiveAnimationType);
 
 	/** Retrieves the Additive Base Pose type for the given Animation Sequence */
 	UFUNCTION(BlueprintPure, Category = "AnimationBlueprintLibrary|Additive")
-	static void GetAdditiveBasePoseType(const UAnimSequence* AnimationSequence, TEnumAsByte<enum EAdditiveBasePoseType>& AdditiveBasePoseType);
+	static UE_API void GetAdditiveBasePoseType(const UAnimSequence* AnimationSequence, TEnumAsByte<enum EAdditiveBasePoseType>& AdditiveBasePoseType);
 
 	/** Sets the Additive Base Pose type for the given Animation Sequence */
 	UFUNCTION(BlueprintCallable, Category = "AnimationBlueprintLibrary|Additive")
-	static void SetAdditiveBasePoseType(UAnimSequence* AnimationSequence, const TEnumAsByte<enum EAdditiveBasePoseType> AdditiveBasePoseType);
+	static UE_API void SetAdditiveBasePoseType(UAnimSequence* AnimationSequence, const TEnumAsByte<enum EAdditiveBasePoseType> AdditiveBasePoseType);
 
 	// Interpolation
 
 	/** Retrieves the Animation Interpolation type for the given Animation Sequence */
 	UFUNCTION(BlueprintPure, Category = "AnimationBlueprintLibrary|Interpolation")
-	static void GetAnimationInterpolationType(const UAnimSequence* AnimationSequence, EAnimInterpolationType& InterpolationType);
+	static UE_API void GetAnimationInterpolationType(const UAnimSequence* AnimationSequence, EAnimInterpolationType& InterpolationType);
 
 	/** Sets the Animation Interpolation type for the given Animation Sequence */
 	UFUNCTION(BlueprintCallable, Category = "AnimationBlueprintLibrary|Interpolation")
-	static void SetAnimationInterpolationType(UAnimSequence* AnimationSequence, EAnimInterpolationType InterpolationType);
+	static UE_API void SetAnimationInterpolationType(UAnimSequence* AnimationSequence, EAnimInterpolationType InterpolationType);
 
 	// Root motion
 
 	/** Checks whether or not Root Motion is Enabled for the given Animation Sequence */
 	UFUNCTION(BlueprintPure, Category = "AnimationBlueprintLibrary|RootMotion")
-	static bool IsRootMotionEnabled(const UAnimSequence* AnimationSequence);
+	static UE_API bool IsRootMotionEnabled(const UAnimSequence* AnimationSequence);
 
 	/** Sets whether or not Root Motion is Enabled for the given Animation Sequence */
 	UFUNCTION(BlueprintCallable, Category = "AnimationBlueprintLibrary|RootMotion")
-	static void SetRootMotionEnabled(UAnimSequence* AnimationSequence, bool bEnabled);
+	static UE_API void SetRootMotionEnabled(UAnimSequence* AnimationSequence, bool bEnabled);
 
 	/** Retrieves the Root Motion Lock Type for the given Animation Sequence */
 	UFUNCTION(BlueprintPure, Category = "AnimationBlueprintLibrary|RootMotion")
-	static void GetRootMotionLockType(const UAnimSequence* AnimationSequence, TEnumAsByte<ERootMotionRootLock::Type>& LockType);
+	static UE_API void GetRootMotionLockType(const UAnimSequence* AnimationSequence, TEnumAsByte<ERootMotionRootLock::Type>& LockType);
 
 	/** Sets the Root Motion Lock Type for the given Animation Sequence */
 	UFUNCTION(BlueprintCallable, Category = "AnimationBlueprintLibrary|RootMotion")
-	static void SetRootMotionLockType(UAnimSequence* AnimationSequence, TEnumAsByte<ERootMotionRootLock::Type> RootMotionLockType);
+	static UE_API void SetRootMotionLockType(UAnimSequence* AnimationSequence, TEnumAsByte<ERootMotionRootLock::Type> RootMotionLockType);
 
 	/** Checks whether or not Root Motion locking is Forced for the given Animation Sequence */
 	UFUNCTION(BlueprintPure, Category = "AnimationBlueprintLibrary|RootMotion")
-	static bool IsRootMotionLockForced(const UAnimSequence* AnimationSequence);
+	static UE_API bool IsRootMotionLockForced(const UAnimSequence* AnimationSequence);
 
 	/** Sets whether or not Root Motion locking is Forced for the given Animation Sequence */
 	UFUNCTION(BlueprintCallable, Category = "AnimationBlueprintLibrary|RootMotion")
-	static void SetIsRootMotionLockForced(UAnimSequence* AnimationSequence, bool bForced);
+	static UE_API void SetIsRootMotionLockForced(UAnimSequence* AnimationSequence, bool bForced);
 
 	// Markers
 
 	/** Retrieves all the Animation Sync Markers for the given Animation Sequence */
 	UFUNCTION(BlueprintPure, Category = "AnimationBlueprintLibrary|MarkerSyncing")
-	static void GetAnimationSyncMarkers(const UAnimSequence* AnimationSequence, TArray<FAnimSyncMarker>& Markers);
+	static UE_API void GetAnimationSyncMarkers(const UAnimSequence* AnimationSequence, TArray<FAnimSyncMarker>& Markers);
 
 	/** Retrieves all the Unique Names for the Animation Sync Markers contained by the given Animation Sequence */
 	UFUNCTION(BlueprintPure, Category = "AnimationBlueprintLibrary|MarkerSyncing")
-	static void GetUniqueMarkerNames(const UAnimSequence* AnimationSequence, TArray<FName>& MarkerNames);
+	static UE_API void GetUniqueMarkerNames(const UAnimSequence* AnimationSequence, TArray<FName>& MarkerNames);
 
 	/** Adds an Animation Sync Marker to Notify track in the given Animation with the corresponding Marker Name and Time */
 	UFUNCTION(BlueprintCallable, Category = "AnimationBlueprintLibrary|MarkerSyncing")
-	static void AddAnimationSyncMarker(UAnimSequence* AnimationSequence, FName MarkerName, float Time, FName NotifyTrackName);
+	static UE_API void AddAnimationSyncMarker(UAnimSequence* AnimationSequence, FName MarkerName, float Time, FName NotifyTrackName);
 
 	/** Checks whether or not the given Marker Name is a valid Animation Sync Marker Name */
 	UFUNCTION(BlueprintPure, Category = "AnimationBlueprintLibrary|Helpers")
-	static bool IsValidAnimationSyncMarkerName(const UAnimSequence* AnimationSequence, FName MarkerName);
+	static UE_API bool IsValidAnimationSyncMarkerName(const UAnimSequence* AnimationSequence, FName MarkerName);
 
 	/** Removes All Animation Sync Marker found within the Animation Sequence whose name matches MarkerName, and returns the number of removed instances */
 	UFUNCTION(BlueprintCallable, Category = "AnimationBlueprintLibrary|MarkerSyncing")
-	static int32 RemoveAnimationSyncMarkersByName(UAnimSequence* AnimationSequence, FName MarkerName);
+	static UE_API int32 RemoveAnimationSyncMarkersByName(UAnimSequence* AnimationSequence, FName MarkerName);
 
 	/** Removes All Animation Sync Marker found within the Animation Sequence that belong to the specific Notify Track, and returns the number of removed instances */	
 	UFUNCTION(BlueprintCallable, Category = "AnimationBlueprintLibrary|MarkerSyncing")
-	static int32 RemoveAnimationSyncMarkersByTrack(UAnimSequence* AnimationSequence, FName NotifyTrackName);
+	static UE_API int32 RemoveAnimationSyncMarkersByTrack(UAnimSequence* AnimationSequence, FName NotifyTrackName);
 
 	/** Removes All Animation Sync Markers found within the Animation Sequence, and returns the number of removed instances */
 	UFUNCTION(BlueprintCallable, Category = "AnimationBlueprintLibrary|MarkerSyncing")
-	static void RemoveAllAnimationSyncMarkers(UAnimSequence* AnimationSequence);
+	static UE_API void RemoveAllAnimationSyncMarkers(UAnimSequence* AnimationSequence);
 
 	// Notifies
 
 	/** Retrieves all Animation Notify Events found within the given Animation Sequence */
 	UFUNCTION(BlueprintPure, Category = "AnimationBlueprintLibrary|NotifyEvents")
-	static void GetAnimationNotifyEvents(const UAnimSequenceBase* AnimationSequenceBase, TArray<FAnimNotifyEvent>& NotifyEvents);
+	static UE_API void GetAnimationNotifyEvents(const UAnimSequenceBase* AnimationSequenceBase, TArray<FAnimNotifyEvent>& NotifyEvents);
 
 	/** Retrieves all Unique Animation Notify Events found within the given Animation Sequence */
 	UFUNCTION(BlueprintPure, Category = "AnimationBlueprintLibrary|NotifyEvents")
-	static void GetAnimationNotifyEventNames(const UAnimSequenceBase* AnimationSequenceBase, TArray<FName>& EventNames);
+	static UE_API void GetAnimationNotifyEventNames(const UAnimSequenceBase* AnimationSequenceBase, TArray<FName>& EventNames);
 
 	/** Adds an Animation Notify Event to Notify track in the given Animation with the given Notify creation data */
 	UFUNCTION(BlueprintCallable, Category = "AnimationBlueprintLibrary|NotifyEvents")
-	static UAnimNotify* AddAnimationNotifyEvent(UAnimSequenceBase* AnimationSequenceBase, FName NotifyTrackName, float StartTime, TSubclassOf<UAnimNotify> NotifyClass);
+	static UE_API UAnimNotify* AddAnimationNotifyEvent(UAnimSequenceBase* AnimationSequenceBase, FName NotifyTrackName, float StartTime, TSubclassOf<UAnimNotify> NotifyClass);
 
 	/** Adds an Animation Notify State Event to Notify track in the given Animation with the given Notify State creation data */
 	UFUNCTION(BlueprintCallable, Category = "AnimationBlueprintLibrary|NotifyEvents")
-	static UAnimNotifyState* AddAnimationNotifyStateEvent(UAnimSequenceBase* AnimationSequenceBase, FName NotifyTrackName, float StartTime, float Duration, TSubclassOf<UAnimNotifyState> NotifyStateClass);
+	static UE_API UAnimNotifyState* AddAnimationNotifyStateEvent(UAnimSequenceBase* AnimationSequenceBase, FName NotifyTrackName, float StartTime, float Duration, TSubclassOf<UAnimNotifyState> NotifyStateClass);
 
 	/** Adds an the specific Animation Notify to the Animation Sequence (requires Notify's outer to be the Animation Sequence) */
 	UFUNCTION(BlueprintCallable, Category = "AnimationBlueprintLibrary|NotifyEvents")
-	static void AddAnimationNotifyEventObject(UAnimSequenceBase* AnimationSequenceBase, float StartTime, UAnimNotify* Notify, FName NotifyTrackName);
+	static UE_API void AddAnimationNotifyEventObject(UAnimSequenceBase* AnimationSequenceBase, float StartTime, UAnimNotify* Notify, FName NotifyTrackName);
 
 	/** Adds an the specific Animation Notify State to the Animation Sequence (requires Notify State's outer to be the Animation Sequence) */
 	UFUNCTION(BlueprintCallable, Category = "AnimationBlueprintLibrary|NotifyEvents")
-	static void AddAnimationNotifyStateEventObject(UAnimSequenceBase* AnimationSequenceBase, float StartTime, float Duration, UAnimNotifyState* NotifyState, FName NotifyTrackName);
+	static UE_API void AddAnimationNotifyStateEventObject(UAnimSequenceBase* AnimationSequenceBase, float StartTime, float Duration, UAnimNotifyState* NotifyState, FName NotifyTrackName);
 
 	/** Removes Animation Notify Events found by Name within the Animation Sequence, and returns the number of removed name instances */
 	UFUNCTION(BlueprintCallable, Category = "AnimationBlueprintLibrary|NotifyEvents")
-	static int32 RemoveAnimationNotifyEventsByName(UAnimSequenceBase* AnimationSequenceBase, FName NotifyName);
+	static UE_API int32 RemoveAnimationNotifyEventsByName(UAnimSequenceBase* AnimationSequenceBase, FName NotifyName);
 
 	/** Removes Animation Notify Events found by Track within the Animation Sequence, and returns the number of removed name instances */
 	UFUNCTION(BlueprintCallable, Category = "AnimationBlueprintLibrary|NotifyEvents")
-	static int32 RemoveAnimationNotifyEventsByTrack(UAnimSequenceBase* AnimationSequenceBase, FName NotifyTrackName);	
+	static UE_API int32 RemoveAnimationNotifyEventsByTrack(UAnimSequenceBase* AnimationSequenceBase, FName NotifyTrackName);	
 
 	/** Replaces animation notifies in the specified Animation Sequence */
 	UFUNCTION(BlueprintCallable, Category = "AnimationBlueprintLibrary|NotifyEvents")
-	static void ReplaceAnimNotifyStates(UAnimSequenceBase* AnimationSequenceBase, TSubclassOf<UAnimNotifyState> OldNotifyClass, TSubclassOf<UAnimNotifyState> NewNotifyClass, FOnNotifyStateReplaced OnNotifyStateReplaced);
+	static UE_API void ReplaceAnimNotifyStates(UAnimSequenceBase* AnimationSequenceBase, TSubclassOf<UAnimNotifyState> OldNotifyClass, TSubclassOf<UAnimNotifyState> NewNotifyClass, FOnNotifyStateReplaced OnNotifyStateReplaced);
 
 	/** Replaces animation notifies in the specified Animation Sequence */
 	UFUNCTION(BlueprintCallable, Category = "AnimationBlueprintLibrary|NotifyEvents")
-	static void ReplaceAnimNotifies(UAnimSequenceBase* AnimationSequenceBase, TSubclassOf<UAnimNotify> OldNotifyClass, TSubclassOf<UAnimNotify> NewNotifyClass, FOnNotifyReplaced OnNotifyReplaced);
+	static UE_API void ReplaceAnimNotifies(UAnimSequenceBase* AnimationSequenceBase, TSubclassOf<UAnimNotify> OldNotifyClass, TSubclassOf<UAnimNotify> NewNotifyClass, FOnNotifyReplaced OnNotifyReplaced);
 
 	/** Copies animation notifies from Src Animation Sequence to Dest. Creates anim notify tracks as necessary. Returns true on success. */
 	UFUNCTION(BlueprintCallable, Category = "AnimationBlueprintLibrary|NotifyEvents")
-	static void CopyAnimNotifiesFromSequence(UAnimSequenceBase* SourceAnimationSequenceBase, UAnimSequenceBase* DestinationAnimationSequenceBase, bool bDeleteExistingNotifies = false);
+	static UE_API void CopyAnimNotifiesFromSequence(UAnimSequenceBase* SourceAnimationSequenceBase, UAnimSequenceBase* DestinationAnimationSequenceBase, bool bDeleteExistingNotifies = false);
 
 	// Notify Tracks
 
 	/** Retrieves all Unique Animation Notify Track Names found within the given Animation Sequence */
 	UFUNCTION(BlueprintPure, Category = "AnimationBlueprintLibrary|AnimationNotifies")
-	static void GetAnimationNotifyTrackNames(const UAnimSequenceBase* AnimationSequenceBase, TArray<FName>& TrackNames);
+	static UE_API void GetAnimationNotifyTrackNames(const UAnimSequenceBase* AnimationSequenceBase, TArray<FName>& TrackNames);
 
 	/** Adds an Animation Notify Track to the Animation Sequence */
 	UFUNCTION(BlueprintCallable, Category = "AnimationBlueprintLibrary|AnimationNotifies")
-	static void AddAnimationNotifyTrack(UAnimSequenceBase* AnimationSequenceBase, FName NotifyTrackName, FLinearColor TrackColor = FLinearColor::White);
+	static UE_API void AddAnimationNotifyTrack(UAnimSequenceBase* AnimationSequenceBase, FName NotifyTrackName, FLinearColor TrackColor = FLinearColor::White);
 
 	/** Removes an Animation Notify Track from Animation Sequence by Name */
 	UFUNCTION(BlueprintCallable, Category = "AnimationBlueprintLibrary|AnimationNotifies")
-	static void RemoveAnimationNotifyTrack(UAnimSequenceBase* AnimationSequenceBase, FName NotifyTrackName);
+	static UE_API void RemoveAnimationNotifyTrack(UAnimSequenceBase* AnimationSequenceBase, FName NotifyTrackName);
 
 	/** Removes All Animation Notify Tracks from Animation Sequence */
 	UFUNCTION(BlueprintCallable, Category = "AnimationBlueprintLibrary|AnimationNotifies")
-	static void RemoveAllAnimationNotifyTracks(UAnimSequenceBase* AnimationSequenceBase);
+	static UE_API void RemoveAllAnimationNotifyTracks(UAnimSequenceBase* AnimationSequenceBase);
 
 	/** Checks whether or not the given Track Name is a valid Animation Notify Track in the Animation Sequence */
 	UFUNCTION(BlueprintPure, Category = "AnimationBlueprintLibrary|Helpers")
-	static bool IsValidAnimNotifyTrackName(const UAnimSequenceBase* AnimationSequenceBase, FName NotifyTrackName);
+	static UE_API bool IsValidAnimNotifyTrackName(const UAnimSequenceBase* AnimationSequenceBase, FName NotifyTrackName);
 
-	static int32 GetTrackIndexForAnimationNotifyTrackName(const UAnimSequenceBase* AnimationSequenceBase, FName NotifyTrackName);
-	static const FAnimNotifyTrack& GetNotifyTrackByName(const UAnimSequenceBase* AnimationSequenceBase, FName NotifyTrackName);
+	static UE_API int32 GetTrackIndexForAnimationNotifyTrackName(const UAnimSequenceBase* AnimationSequenceBase, FName NotifyTrackName);
+	static UE_API const FAnimNotifyTrack& GetNotifyTrackByName(const UAnimSequenceBase* AnimationSequenceBase, FName NotifyTrackName);
 
 	/** Returns the actual trigger time for a NotifyEvent */
 	UFUNCTION(BlueprintPure, Category = "AnimationBlueprintLibrary|AnimationNotifies")
-	static float GetAnimNotifyEventTriggerTime(const FAnimNotifyEvent& NotifyEvent);
+	static UE_API float GetAnimNotifyEventTriggerTime(const FAnimNotifyEvent& NotifyEvent);
 	
 	/** Returns the duration for a NotifyEvent, only non-zero for Anim Notify States */
 	UFUNCTION(BlueprintPure, Category = "AnimationBlueprintLibrary|AnimationNotifies")
-	static float GetAnimNotifyEventDuration(const FAnimNotifyEvent& NotifyEvent);
+	static UE_API float GetAnimNotifyEventDuration(const FAnimNotifyEvent& NotifyEvent);
 
 	/** Retrieves all Animation Sync Markers for the given Notify Track Name from the given Animation Sequence */
 	UFUNCTION(BlueprintPure, Category = "AnimationBlueprintLibrary|MarkerSyncing")
-	static void GetAnimationSyncMarkersForTrack(const UAnimSequence* AnimationSequence, FName NotifyTrackName, TArray<FAnimSyncMarker>& Markers);
+	static UE_API void GetAnimationSyncMarkersForTrack(const UAnimSequence* AnimationSequence, FName NotifyTrackName, TArray<FAnimSyncMarker>& Markers);
 
 	/** Retrieves all Animation Notify Events for the given Notify Track Name from the given Animation Sequence */
 	UFUNCTION(BlueprintPure, Category = "AnimationBlueprintLibrary|NotifyEvents")
-	static void GetAnimationNotifyEventsForTrack(const UAnimSequenceBase* AnimationSequenceBase, FName NotifyTrackName, TArray<FAnimNotifyEvent>& Events);
+	static UE_API void GetAnimationNotifyEventsForTrack(const UAnimSequenceBase* AnimationSequenceBase, FName NotifyTrackName, TArray<FAnimNotifyEvent>& Events);
 
 	// Curves
 
 	/** Adds an Animation Curve by Type and Name to the given Animation Sequence */
 	UFUNCTION(BlueprintCallable, Category = "AnimationBlueprintLibrary|Curves")
-	static void AddCurve(UAnimSequenceBase* AnimationSequenceBase, FName CurveName, ERawCurveTrackTypes CurveType = ERawCurveTrackTypes::RCT_Float, bool bMetaDataCurve = false);
+	static UE_API void AddCurve(UAnimSequenceBase* AnimationSequenceBase, FName CurveName, ERawCurveTrackTypes CurveType = ERawCurveTrackTypes::RCT_Float, bool bMetaDataCurve = false);
 
 	/** Removes an Animation Curve by Name from the given Animation Sequence (Raw Animation Curves [Names] may not be removed from the Skeleton) */
 	UFUNCTION(BlueprintCallable, Category = "AnimationBlueprintLibrary|Curves")
-	static void RemoveCurve(UAnimSequenceBase* AnimationSequenceBase, FName CurveName, bool bRemoveNameFromSkeleton = false);
+	static UE_API void RemoveCurve(UAnimSequenceBase* AnimationSequenceBase, FName CurveName, bool bRemoveNameFromSkeleton = false);
 
 	/** Removes all Animation Curve Data from the given Animation Sequence (Raw Animation Curves [Names] may not be removed from the Skeleton) */
 	UFUNCTION(BlueprintCallable, Category = "AnimationBlueprintLibrary|Curves")
-	static void RemoveAllCurveData(UAnimSequenceBase* AnimationSequenceBase);
+	static UE_API void RemoveAllCurveData(UAnimSequenceBase* AnimationSequenceBase);
 
 	/** Adds a Transformation Key to the specified Animation Curve inside of the given Animation Sequence */
 	UFUNCTION(BlueprintCallable, Category = "AnimationBlueprintLibrary|Curves")
-	static void AddTransformationCurveKey(UAnimSequenceBase* AnimationSequenceBase, FName CurveName, const float Time, const FTransform& Transform);
+	static UE_API void AddTransformationCurveKey(UAnimSequenceBase* AnimationSequenceBase, FName CurveName, const float Time, const FTransform& Transform);
 
 	/** Adds a multiple of Transformation Keys to the specified Animation Curve inside of the given Animation Sequence */
 	UFUNCTION(BlueprintCallable, Category = "AnimationBlueprintLibrary|Curves")
-	static void AddTransformationCurveKeys(UAnimSequenceBase* AnimationSequenceBase, FName CurveName, const TArray<float>& Times, const TArray<FTransform>& Transforms);
+	static UE_API void AddTransformationCurveKeys(UAnimSequenceBase* AnimationSequenceBase, FName CurveName, const TArray<float>& Times, const TArray<FTransform>& Transforms);
 
 	/** Adds a Float Key to the specified Animation Curve inside of the given Animation Sequence */
 	UFUNCTION(BlueprintCallable, Category = "AnimationBlueprintLibrary|Curves")
-	static void AddFloatCurveKey(UAnimSequenceBase* AnimationSequenceBase, FName CurveName, const float Time, const float Value);
+	static UE_API void AddFloatCurveKey(UAnimSequenceBase* AnimationSequenceBase, FName CurveName, const float Time, const float Value);
 
 	/** Adds a multiple of Float Keys to the specified Animation Curve inside of the given Animation Sequence */
 	UFUNCTION(BlueprintCallable, Category = "AnimationBlueprintLibrary|Curves")
-	static void AddFloatCurveKeys(UAnimSequenceBase* AnimationSequenceBase, FName CurveName, const TArray<float>& Times, const TArray<float>& Values);
+	static UE_API void AddFloatCurveKeys(UAnimSequenceBase* AnimationSequenceBase, FName CurveName, const TArray<float>& Times, const TArray<float>& Values);
 
 	/** Adds a Vector Key to the specified Animation Curve inside of the given Animation Sequence */
 	UFUNCTION(BlueprintCallable, Category = "AnimationBlueprintLibrary|Curves")
-	static void AddVectorCurveKey(UAnimSequenceBase* AnimationSequenceBase, FName CurveName, const float Time, const FVector Vector);
+	static UE_API void AddVectorCurveKey(UAnimSequenceBase* AnimationSequenceBase, FName CurveName, const float Time, const FVector Vector);
 
 	/** Adds a multiple of Vector Keys to the specified Animation Curve inside of the given Animation Sequence */
 	UFUNCTION(BlueprintCallable, Category = "AnimationBlueprintLibrary|Curves")
-	static void AddVectorCurveKeys(UAnimSequenceBase* AnimationSequenceBase, FName CurveName, const TArray<float>& Times, const TArray<FVector>& Vectors);
+	static UE_API void AddVectorCurveKeys(UAnimSequenceBase* AnimationSequenceBase, FName CurveName, const TArray<float>& Times, const TArray<FVector>& Vectors);
 
 	// Curve helper functions
 	template <typename DataType, typename CurveClass, ERawCurveTrackTypes CurveType>
 	static void AddCurveKeysInternal(UAnimSequenceBase* AnimationSequenceBase, FName CurveName, const TArray<float>& Times, const TArray<DataType>& KeyData);
 
 	// Returns true if successfully added, false if it was already existing
-	static bool AddCurveInternal(UAnimSequenceBase* AnimationSequenceBase, FName CurveName, int32 CurveFlags, ERawCurveTrackTypes SupportedCurveType);
-	static bool RemoveCurveInternal(UAnimSequenceBase* AnimationSequenceBase, FName CurveName, ERawCurveTrackTypes SupportedCurveType);
+	static UE_API bool AddCurveInternal(UAnimSequenceBase* AnimationSequenceBase, FName CurveName, int32 CurveFlags, ERawCurveTrackTypes SupportedCurveType);
+	static UE_API bool RemoveCurveInternal(UAnimSequenceBase* AnimationSequenceBase, FName CurveName, ERawCurveTrackTypes SupportedCurveType);
 
 	/** Checks whether or not the given Bone Name exist on the Skeleton referenced by the given Animation Sequence */
 	UFUNCTION(BlueprintCallable, Category = "AnimationBlueprintLibrary|Skeleton")
-	static void DoesBoneNameExist(UAnimSequence* AnimationSequence, FName BoneName, bool& bExists);
+	static UE_API void DoesBoneNameExist(UAnimSequence* AnimationSequence, FName BoneName, bool& bExists);
 
-	static bool DoesBoneNameExistInternal(USkeleton* Skeleton, FName BoneName);
+	static UE_API bool DoesBoneNameExistInternal(USkeleton* Skeleton, FName BoneName);
 
 	UE_DEPRECATED(5.3, "This function is no longer used.")
 	static bool DoesBoneCurveNameExistInternal(USkeleton* Skeleton, FName BoneName) { return false; }
 
 	/** Retrieves, a multiple of, Float Key(s) from the specified Animation Curve inside of the given Animation Sequence */
 	UFUNCTION(BlueprintCallable, Category = "AnimationBlueprintLibrary|Curves")
-	static void GetFloatKeys(UAnimSequenceBase* AnimationSequenceBase, FName CurveName, TArray<float>& Times, TArray<float>& Values);
+	static UE_API void GetFloatKeys(UAnimSequenceBase* AnimationSequenceBase, FName CurveName, TArray<float>& Times, TArray<float>& Values);
 
 	/** Retrieves, a multiple of, Vector Key(s) from the specified Animation Curve inside of the given Animation Sequence */
 	UFUNCTION(BlueprintCallable, Category = "AnimationBlueprintLibrary|Curves")
-	static void GetVectorKeys(UAnimSequenceBase* AnimationSequenceBase, FName CurveName, TArray<float>& Times, TArray<FVector>& Values);
+	static UE_API void GetVectorKeys(UAnimSequenceBase* AnimationSequenceBase, FName CurveName, TArray<float>& Times, TArray<FVector>& Values);
 
 	/** Retrieves, a multiple of, Transformation Key(s) from the specified Animation Curve inside of the given Animation Sequence */
 	UFUNCTION(BlueprintCallable, Category = "AnimationBlueprintLibrary|Curves")
-	static void GetTransformationKeys(UAnimSequenceBase* AnimationSequenceBase, FName CurveName, TArray<float>& Times, TArray<FTransform>& Values);
+	static UE_API void GetTransformationKeys(UAnimSequenceBase* AnimationSequenceBase, FName CurveName, TArray<float>& Times, TArray<FTransform>& Values);
 
 	/** Retrieves an evaluated float value for a given time from the specified Animation Curve inside of the given Animation Sequence */
 	UFUNCTION(BlueprintCallable, Category = "AnimationBlueprintLibrary|Curves")
-	static float GetFloatValueAtTime(UAnimSequenceBase* AnimationSequenceBase, FName CurveName, float Time);
+	static UE_API float GetFloatValueAtTime(UAnimSequenceBase* AnimationSequenceBase, FName CurveName, float Time);
 	
 	template <typename DataType, typename CurveClass, ERawCurveTrackTypes CurveType>
 	static void GetCurveKeysInternal(UAnimSequenceBase* AnimationSequenceBase, FName CurveName, TArray<float>& Times, TArray<DataType>& KeyData);
@@ -401,7 +403,7 @@ public:
 	 *  @param bFinalize : If you set this to true, it will trigger compression. If you set bFinalize to be false, you'll have to manually trigger Finalize. 
 	 */
 	UFUNCTION(BlueprintCallable, Category = "AnimationBlueprintLibrary|Bones")
-	static void RemoveBoneAnimation(UAnimSequence* AnimationSequence, FName BoneName, bool bIncludeChildren = true, bool bFinalize = true);
+	static UE_API void RemoveBoneAnimation(UAnimSequence* AnimationSequence, FName BoneName, bool bIncludeChildren = true, bool bFinalize = true);
 
 	/** Removes an Animation Curve by Name from the given Animation Sequence (Raw Animation Curves [Names] may not be removed from the Skeleton) 
 	 * 
@@ -413,11 +415,11 @@ public:
 	 *  @param bFinalize : If you set this to true, it will trigger compression. If you set bFinalize to be false, you'll have to manually trigger Finalize.
 	 */
 	UFUNCTION(BlueprintCallable, Category = "AnimationBlueprintLibrary|Bones", meta=(AutoCreateRefTerm="ChildrenExcluded"))
-	static void RemoveBoneSelectiveAnimation(UAnimSequence* AnimationSequence, FName BoneName, TArray<FName> ChildrenExcluded, bool bIncludeChildren = true, bool bExcludeRecursively = false, bool bFinalize = true);
+	static UE_API void RemoveBoneSelectiveAnimation(UAnimSequence* AnimationSequence, FName BoneName, TArray<FName> ChildrenExcluded, bool bIncludeChildren = true, bool bExcludeRecursively = false, bool bFinalize = true);
 
 	/** Removes all Animation Bone Track Data from the given Animation Sequence */
 	UFUNCTION(BlueprintCallable, Category = "AnimationBlueprintLibrary|Curves")
-	static void RemoveAllBoneAnimation(UAnimSequence* AnimationSequence);
+	static UE_API void RemoveAllBoneAnimation(UAnimSequence* AnimationSequence);
 
 	/** Apply all the changes made to Bone Tracks to Finalize. This triggers recompression. Note that this is expensive, but will require to get correct compressed data */
 	UE_DEPRECATED(5.0, "FinalizeBoneAnimation has been deprecated, use UAnimDataController instead")
@@ -428,7 +430,7 @@ public:
 
 	/** Checks whether or not the given Curve Name exist on the Skeleton referenced by the given Animation Sequence */
 	UFUNCTION(BlueprintCallable, Category = "AnimationBlueprintLibrary|Curves")
-	static bool DoesCurveExist(UAnimSequenceBase* AnimationSequenceBase, FName CurveName, ERawCurveTrackTypes CurveType);
+	static UE_API bool DoesCurveExist(UAnimSequenceBase* AnimationSequenceBase, FName CurveName, ERawCurveTrackTypes CurveType);
 
 	UE_DEPRECATED(5.3, "This function is no longer used.")
 	static bool DoesSmartNameExist(UAnimSequence* AnimationSequence, FName Name) { return false; }
@@ -447,120 +449,120 @@ public:
 	UE_DEPRECATED(5.3, "This function is no longer used.")
 	static FName RetrieveContainerNameForCurve(const UAnimSequence* AnimationSequence, FName CurveName) { return NAME_None; }
 	
-	static ERawCurveTrackTypes RetrieveCurveTypeForCurve(const UAnimSequenceBase* AnimationSequenceBase, FName CurveName);
+	static UE_API ERawCurveTrackTypes RetrieveCurveTypeForCurve(const UAnimSequenceBase* AnimationSequenceBase, FName CurveName);
 
 	// MetaData
 
 	/** Creates and Adds an instance of the specified MetaData Class to the given Animation Asset */
 	UFUNCTION(BlueprintCallable, Category = "AnimationBlueprintLibrary|MetaData")
-	static void AddMetaData(UAnimationAsset* AnimationAsset, TSubclassOf<UAnimMetaData> MetaDataClass, UAnimMetaData*& MetaDataInstance);
+	static UE_API void AddMetaData(UAnimationAsset* AnimationAsset, TSubclassOf<UAnimMetaData> MetaDataClass, UAnimMetaData*& MetaDataInstance);
 
 	/** Adds an instance of the specified MetaData Class to the given Animation Asset (requires MetaDataObject's outer to be the Animation Asset) */
 	UFUNCTION(BlueprintCallable, Category = "AnimationBlueprintLibrary|MetaData")
-	static void AddMetaDataObject(UAnimationAsset* AnimationAsset, UAnimMetaData* MetaDataObject);
+	static UE_API void AddMetaDataObject(UAnimationAsset* AnimationAsset, UAnimMetaData* MetaDataObject);
 
 	/** Removes all Meta Data from the given Animation Asset */
 	UFUNCTION(BlueprintCallable, Category = "AnimationBlueprintLibrary|MetaData")
-	static void RemoveAllMetaData(UAnimationAsset* AnimationAsset);
+	static UE_API void RemoveAllMetaData(UAnimationAsset* AnimationAsset);
 
 	/** Removes the specified Meta Data Instance from the given Animation Asset */
 	UFUNCTION(BlueprintCallable, Category = "AnimationBlueprintLibrary|MetaData")
-	static void RemoveMetaData(UAnimationAsset* AnimationAsset, UAnimMetaData* MetaDataObject);
+	static UE_API void RemoveMetaData(UAnimationAsset* AnimationAsset, UAnimMetaData* MetaDataObject);
 
 	/** Removes all Meta Data Instance of the specified Class from the given Animation Asset */
 	UFUNCTION(BlueprintCallable, Category = "AnimationBlueprintLibrary|MetaData")
-	static void RemoveMetaDataOfClass(UAnimationAsset* AnimationAsset, TSubclassOf<UAnimMetaData> MetaDataClass);
+	static UE_API void RemoveMetaDataOfClass(UAnimationAsset* AnimationAsset, TSubclassOf<UAnimMetaData> MetaDataClass);
 
 	/** Retrieves all Meta Data Instances from the given Animation Asset */
 	UFUNCTION(BlueprintPure, Category = "AnimationBlueprintLibrary|MetaData")
-	static void GetMetaData(const UAnimationAsset* AnimationAsset, TArray<UAnimMetaData*>& MetaData);
+	static UE_API void GetMetaData(const UAnimationAsset* AnimationAsset, TArray<UAnimMetaData*>& MetaData);
 
 	/** Retrieves all Meta Data Instances from the given Animation Asset */
 	UFUNCTION(BlueprintPure, Category = "AnimationBlueprintLibrary|MetaData")
-	static void GetMetaDataOfClass(const UAnimationAsset* AnimationAsset, TSubclassOf<UAnimMetaData> MetaDataClass, TArray<UAnimMetaData*>& MetaDataOfClass);
+	static UE_API void GetMetaDataOfClass(const UAnimationAsset* AnimationAsset, TSubclassOf<UAnimMetaData> MetaDataClass, TArray<UAnimMetaData*>& MetaDataOfClass);
 
 	/** Checks whether or not the given Animation Asset contains Meta Data Instance of the specified Meta Data Class */
 	UFUNCTION(BlueprintPure, Category = "AnimationBlueprintLibrary|MetaData")
-	static bool ContainsMetaDataOfClass(const UAnimationAsset* AnimationAsset, TSubclassOf<UAnimMetaData> MetaDataClass);
+	static UE_API bool ContainsMetaDataOfClass(const UAnimationAsset* AnimationAsset, TSubclassOf<UAnimMetaData> MetaDataClass);
 
 	// Poses
 
 	/** Retrieves Bone Pose data for the given Bone Name at the specified Time from the given Animation Sequence */
 	UE_DEPRECATED(5.2, "GetBonePosesForTime is deprecated, use AnimPose or AnimationDataModel interface directly")
 	UFUNCTION(BlueprintPure, Category = "AnimationBlueprintLibrary|Pose")
-	static void GetBonePoseForTime(const UAnimSequenceBase* AnimationSequenceBase, FName BoneName, float Time, bool bExtractRootMotion, FTransform& Pose);
+	static UE_API void GetBonePoseForTime(const UAnimSequenceBase* AnimationSequenceBase, FName BoneName, float Time, bool bExtractRootMotion, FTransform& Pose);
 
 	/** Retrieves Bone Pose data for the given Bone Name at the specified Frame from the given Animation Sequence */
 	UE_DEPRECATED(5.2, "GetBonePosesForTime is deprecated, use AnimPose or AnimationDataModel interface directly")
 	UFUNCTION(BlueprintPure, Category = "AnimationBlueprintLibrary|Pose")
-	static void GetBonePoseForFrame(const UAnimSequenceBase* AnimationSequenceBase, FName BoneName, int32 Frame, bool bExtractRootMotion, FTransform& Pose);
+	static UE_API void GetBonePoseForFrame(const UAnimSequenceBase* AnimationSequenceBase, FName BoneName, int32 Frame, bool bExtractRootMotion, FTransform& Pose);
 
 	/** Retrieves Bone Pose data for the given Bone Names at the specified Time from the given Animation Sequence */
 	UE_DEPRECATED(5.2, "GetBonePosesForTime is deprecated, use AnimPose or AnimationDataModel interface directly")
 	UFUNCTION(BlueprintPure, Category = "AnimationBlueprintLibrary|Pose")
-	static void GetBonePosesForTime(const UAnimSequenceBase* AnimationSequenceBase, TArray<FName> BoneNames, float Time, bool bExtractRootMotion, TArray<FTransform>& Poses, const USkeletalMesh* PreviewMesh = nullptr);
+	static UE_API void GetBonePosesForTime(const UAnimSequenceBase* AnimationSequenceBase, TArray<FName> BoneNames, float Time, bool bExtractRootMotion, TArray<FTransform>& Poses, const USkeletalMesh* PreviewMesh = nullptr);
 
 	/** Retrieves Bone Pose data for the given Bone Names at the specified Frame from the given Animation Sequence */
 	UE_DEPRECATED(5.2, "GetBonePosesForTime is deprecated, use AnimPose or AnimationDataModel interface directly")
 	UFUNCTION(BlueprintPure, Category = "AnimationBlueprintLibrary|Pose")
-	static void GetBonePosesForFrame(const UAnimSequenceBase* AnimationSequenceBase, TArray<FName> BoneNames, int32 Frame, bool bExtractRootMotion, TArray<FTransform>& Poses, const USkeletalMesh* PreviewMesh = nullptr);
+	static UE_API void GetBonePosesForFrame(const UAnimSequenceBase* AnimationSequenceBase, TArray<FName> BoneNames, int32 Frame, bool bExtractRootMotion, TArray<FTransform>& Poses, const USkeletalMesh* PreviewMesh = nullptr);
 
 	// Virtual bones
 
 	/** Adds a Virtual Bone between the Source and Target Bones to the given Animation Sequence */
 	UFUNCTION(BlueprintCallable, Category = "AnimationBlueprintLibrary|VirtualBones")
-	static void AddVirtualBone(const UAnimSequence* AnimationSequence, FName SourceBoneName, FName TargetBoneName, FName& VirtualBoneName);
+	static UE_API void AddVirtualBone(const UAnimSequence* AnimationSequence, FName SourceBoneName, FName TargetBoneName, FName& VirtualBoneName);
 
 	/** Removes a Virtual Bone with the specified Bone Name from the given Animation Sequence */
 	UFUNCTION(BlueprintCallable, Category = "AnimationBlueprintLibrary|VirtualBones")
-	static void RemoveVirtualBone(const UAnimSequence* AnimationSequence, FName VirtualBoneName);
+	static UE_API void RemoveVirtualBone(const UAnimSequence* AnimationSequence, FName VirtualBoneName);
 
 	/** Removes Virtual Bones with the specified Bone Names from the given Animation Sequence */
 	UFUNCTION(BlueprintCallable, Category = "AnimationBlueprintLibrary|VirtualBones")
-	static void RemoveVirtualBones(const UAnimSequence* AnimationSequence, TArray<FName> VirtualBoneNames);
+	static UE_API void RemoveVirtualBones(const UAnimSequence* AnimationSequence, TArray<FName> VirtualBoneNames);
 
 	/** Removes all Virtual Bones from the given Animation Sequence */
 	UFUNCTION(BlueprintCallable, Category = "AnimationBlueprintLibrary|VirtualBones")
-	static void RemoveAllVirtualBones(const UAnimSequence* AnimationSequence);
+	static UE_API void RemoveAllVirtualBones(const UAnimSequence* AnimationSequence);
 
-	static bool DoesVirtualBoneNameExistInternal(USkeleton* Skeleton, FName BoneName);
+	static UE_API bool DoesVirtualBoneNameExistInternal(USkeleton* Skeleton, FName BoneName);
 
 	// Misc
 
 	/** Retrieves the Length of the given Animation Sequence */
 	UFUNCTION(BlueprintPure, Category = "AnimationBlueprintLibrary|Animation")
-	static void GetSequenceLength(const UAnimSequenceBase* AnimationSequenceBase, float& Length);
+	static UE_API void GetSequenceLength(const UAnimSequenceBase* AnimationSequenceBase, float& Length);
 
 	/** Retrieves the (Play) Rate Scale of the given Animation Sequence */
 	UFUNCTION(BlueprintPure, Category = "AnimationBlueprintLibrary|Animation")
-	static void GetRateScale(const UAnimSequenceBase* AnimationSequenceBase, float& RateScale);
+	static UE_API void GetRateScale(const UAnimSequenceBase* AnimationSequenceBase, float& RateScale);
 
 	/** Sets the (Play) Rate Scale for the given Animation Sequence */
 	UFUNCTION(BlueprintCallable, Category = "AnimationBlueprintLibrary|Animation")
-	static void SetRateScale(UAnimSequenceBase* AnimationSequenceBase, float RateScale);
+	static UE_API void SetRateScale(UAnimSequenceBase* AnimationSequenceBase, float RateScale);
 
 	/** Retrieves the Frame Index at the specified Time Value for the given Animation Sequence */
 	UFUNCTION(BlueprintPure, Category = "AnimationBlueprintLibrary|Helpers")
-	static void GetFrameAtTime(const UAnimSequenceBase* AnimationSequenceBase, const float Time, int32& Frame);
+	static UE_API void GetFrameAtTime(const UAnimSequenceBase* AnimationSequenceBase, const float Time, int32& Frame);
 
 	/** Retrieves the Time Value at the specified Frame Indexfor the given Animation Sequence */
 	UFUNCTION(BlueprintPure, Category = "AnimationBlueprintLibrary|Helpers")
-	static void GetTimeAtFrame(const UAnimSequenceBase* AnimationSequenceBase, const int32 Frame, float& Time);
+	static UE_API void GetTimeAtFrame(const UAnimSequenceBase* AnimationSequenceBase, const int32 Frame, float& Time);
 	
-	static float GetTimeAtFrameInternal(const UAnimSequenceBase* AnimationSequenceBase, const int32 Frame);
+	static UE_API float GetTimeAtFrameInternal(const UAnimSequenceBase* AnimationSequenceBase, const int32 Frame);
 
 	/** Checks whether or not the given Time Value lies within the given Animation Sequence's Length */
 	UFUNCTION(BlueprintPure, Category = "AnimationBlueprintLibrary|Helpers")
-	static void IsValidTime(const UAnimSequenceBase* AnimationSequenceBase, const float Time, bool& IsValid);
+	static UE_API void IsValidTime(const UAnimSequenceBase* AnimationSequenceBase, const float Time, bool& IsValid);
 
-	static bool IsValidTimeInternal(const UAnimSequenceBase* AnimationSequenceBase, const float Time);
+	static UE_API bool IsValidTimeInternal(const UAnimSequenceBase* AnimationSequenceBase, const float Time);
 
 	/**
 	 * Given an animation sequence determine if any bones in that animation sequence has timecode attribute data.
 	 *
 	 * @return The bone name if found; otherwise none value will be returned.
 	 */
-	static FName FindBoneNameWithTimecodeAttributes(const UAnimSequenceBase* AnimSequenceBase);
+	static UE_API FName FindBoneNameWithTimecodeAttributes(const UAnimSequenceBase* AnimSequenceBase);
 
 	/** Evaluates timecode attributes (e.g. "TCFrame", "TCSecond", etc.) of the root bone and returns the resulting qualified frame time.
 	 *
@@ -572,7 +574,7 @@ public:
 	 *  @return: true if the root bone had timecode attributes that could be evaluated and a qualified frame time was set, or false otherwise.
 	 */
 	UFUNCTION(BlueprintPure, Category = "AnimationBlueprintLibrary|Helpers")
-	static bool EvaluateRootBoneTimecodeAttributesAtTime(const UAnimSequenceBase* AnimationSequenceBase, const float EvalTime, FQualifiedFrameTime& OutQualifiedFrameTime);
+	static UE_API bool EvaluateRootBoneTimecodeAttributesAtTime(const UAnimSequenceBase* AnimationSequenceBase, const float EvalTime, FQualifiedFrameTime& OutQualifiedFrameTime);
 
 	/** Evaluates timecode attributes (e.g. "TCFrame", "TCSecond", etc.) and TCSlate of the root bone and returns the resulting qualified frame time.
 	 *
@@ -586,7 +588,7 @@ public:
 	 *  @return: true if the root bone had timecode attributes that could be evaluated and a qualified frame time was set, or false otherwise.
 	 */
 	UFUNCTION(BlueprintPure, Category = "AnimationBlueprintLibrary|Helpers")
-	static bool EvaluateBoneTimecodeAndSlateAttributesAtTime(const FName BoneName, const UAnimSequenceBase* AnimationSequenceBase, const float EvalTime, FQualifiedFrameTime& OutQualifiedFrameTime, FString& Slate);
+	static UE_API bool EvaluateBoneTimecodeAndSlateAttributesAtTime(const FName BoneName, const UAnimSequenceBase* AnimationSequenceBase, const float EvalTime, FQualifiedFrameTime& OutQualifiedFrameTime, FString& Slate);
 
 	/** Evaluates the subframe timecode attribute (e.g. "TCSubframe") of the root bone and returns the resulting value.
 	 *
@@ -604,19 +606,19 @@ public:
 	 *  @return: true if the root bone had a subframe timecode attribute that could be evaluated and a value was set, or false otherwise.
 	 */
 	UFUNCTION(BlueprintPure, Category = "AnimationBlueprintLibrary|Helpers")
-	static bool EvaluateRootBoneTimecodeSubframeAttributeAtTime(const UAnimSequenceBase* AnimationSequenceBase, const float EvalTime, float& OutSubframe);
+	static UE_API bool EvaluateRootBoneTimecodeSubframeAttributeAtTime(const UAnimSequenceBase* AnimationSequenceBase, const float EvalTime, float& OutSubframe);
 
 	/** Finds the Bone Path from the given Bone to the Root Bone */
 	UFUNCTION(BlueprintPure, Category = "AnimationBlueprintLibrary|Helpers")
-	static void FindBonePathToRoot(const UAnimSequenceBase* AnimationSequenceBase, FName BoneName, TArray<FName>& BonePath);
+	static UE_API void FindBonePathToRoot(const UAnimSequenceBase* AnimationSequenceBase, FName BoneName, TArray<FName>& BonePath);
 
 	/** Returns all Animation Graphs contained by the provided Animation Blueprint */
 	UFUNCTION(BlueprintCallable, Category=Animation, meta=(ScriptMethod))
-	static void GetAnimationGraphs(UAnimBlueprint* AnimationBlueprint, TArray<UAnimationGraph*>& AnimationGraphs);
+	static UE_API void GetAnimationGraphs(UAnimBlueprint* AnimationBlueprint, TArray<UAnimationGraph*>& AnimationGraphs);
 
 	/** Returns all Animation Graph Nodes of the provided Node Class contained by the Animation Blueprint */
 	UFUNCTION(BlueprintCallable, Category=Animation, meta=(ScriptMethod))
-	static void GetNodesOfClass(UAnimBlueprint* AnimationBlueprint, TSubclassOf<UAnimGraphNode_Base> NodeClass, TArray<UAnimGraphNode_Base*>& GraphNodes, bool bIncludeChildClasses = true);
+	static UE_API void GetNodesOfClass(UAnimBlueprint* AnimationBlueprint, TSubclassOf<UAnimGraphNode_Base> NodeClass, TArray<UAnimGraphNode_Base*>& GraphNodes, bool bIncludeChildClasses = true);
 
 	/**
 	 * Adds an Animation Asset override for the provided AnimationBlueprint, replacing any instance of Target with Override
@@ -627,5 +629,7 @@ public:
 	 * @param bPrintAppliedOverrides		Flag whether or not to print the applied overrides
 	 */
 	UFUNCTION(BlueprintCallable, Category=Animation, meta = (ScriptMethod))
-	static void AddNodeAssetOverride(UAnimBlueprint* AnimBlueprint, const UAnimationAsset* Target, UAnimationAsset* Override, bool bPrintAppliedOverrides = false);
+	static UE_API void AddNodeAssetOverride(UAnimBlueprint* AnimBlueprint, const UAnimationAsset* Target, UAnimationAsset* Override, bool bPrintAppliedOverrides = false);
 };
+
+#undef UE_API
diff --git a/Engine/Source/Editor/AnimationEditMode/Public/AnimationEditMode.h b/Engine/Source/Editor/AnimationEditMode/Public/AnimationEditMode.h
--- a/Engine/Source/Editor/AnimationEditMode/Public/AnimationEditMode.h
+++ b/Engine/Source/Editor/AnimationEditMode/Public/AnimationEditMode.h
@@ -14,6 +14,8 @@
 
 #include "AnimationEditMode.generated.h"
 
+#define UE_API ANIMATIONEDITMODE_API
+
 class FAnimationEditMode;
 class FText;
 
@@ -40,17 +42,19 @@ private:
 	friend FAnimationEditMode;
 };
 
-class ANIMATIONEDITMODE_API FAnimationEditMode : public FEdMode, public IAnimationEditContext
+class FAnimationEditMode : public FEdMode, public IAnimationEditContext
 {
 public:
-	FAnimationEditMode();
+	UE_API FAnimationEditMode();
 
 	FAnimationEditMode(const FAnimationEditMode&) = delete;
 	FAnimationEditMode& operator=(const FAnimationEditMode&) = delete;
 
-	virtual void Enter() override;
-	virtual void Exit() override;
-	virtual void AddReferencedObjects(FReferenceCollector& Collector) override;
+	UE_API virtual void Enter() override;
+	UE_API virtual void Exit() override;
+	UE_API virtual void AddReferencedObjects(FReferenceCollector& Collector) override;
 private:
 	TObjectPtr<UAnimationEditModeContext> AnimationEditModeContext;
 };
+
+#undef UE_API
diff --git a/Engine/Source/Editor/AnimationEditor/Public/AnimationToolMenuContext.h b/Engine/Source/Editor/AnimationEditor/Public/AnimationToolMenuContext.h
--- a/Engine/Source/Editor/AnimationEditor/Public/AnimationToolMenuContext.h
+++ b/Engine/Source/Editor/AnimationEditor/Public/AnimationToolMenuContext.h
@@ -8,11 +8,11 @@
 
 class IAnimationEditor;
 
-UCLASS()
-class ANIMATIONEDITOR_API UAnimationToolMenuContext : public UObject
+UCLASS(MinimalAPI)
+class UAnimationToolMenuContext : public UObject
 {
 	GENERATED_BODY()
 
 public:
 	TWeakPtr<IAnimationEditor> AnimationEditor;
-};
\ No newline at end of file
+};
diff --git a/Engine/Source/Editor/AnimationEditorWidgets/Public/Overrides/OverrideStatusDetailsDisplayManager.h b/Engine/Source/Editor/AnimationEditorWidgets/Public/Overrides/OverrideStatusDetailsDisplayManager.h
--- a/Engine/Source/Editor/AnimationEditorWidgets/Public/Overrides/OverrideStatusDetailsDisplayManager.h
+++ b/Engine/Source/Editor/AnimationEditorWidgets/Public/Overrides/OverrideStatusDetailsDisplayManager.h
@@ -8,37 +8,39 @@
 #include "Overrides/OverrideStatusDetailsWidgetBuilder.h"
 #include "Overrides/OverrideStatusWidgetMenuBuilder.h"
 
+#define UE_API ANIMATIONEDITORWIDGETS_API
+
 class FOverrideStatusDetailsWidgetBuilder;
 
 /**
  * The display manager is used to determine how the details view should behave when using an object
  * filter. In this case the display manager is used to set the property-updated-widget (the override status widget).
  */
-class ANIMATIONEDITORWIDGETS_API FOverrideStatusDetailsDisplayManager : public FDetailsDisplayManager
+class FOverrideStatusDetailsDisplayManager : public FDetailsDisplayManager
 {
 public:
 
-	virtual ~FOverrideStatusDetailsDisplayManager() override;
+	UE_API virtual ~FOverrideStatusDetailsDisplayManager() override;
 
 	/**
 	 * Returns a @code bool @endcode indicating whether this @code DetailsViewObjectFilter @endcode instance
 	 * has a category menu
 	 */
-	virtual bool ShouldShowCategoryMenu() override;
+	UE_API virtual bool ShouldShowCategoryMenu() override;
 
 	/**
 	 * Returns the @code const FDetailsViewStyleKey& @endcode that is the Key to the current FDetailsViewStyle style
 	 */
-	virtual const FDetailsViewStyleKey& GetDetailsViewStyleKey() const override;
+	UE_API virtual const FDetailsViewStyleKey& GetDetailsViewStyleKey() const override;
 
 	// Returns true if this manager can construct the property updated widget
-	virtual bool CanConstructPropertyUpdatedWidgetBuilder() const override;
+	UE_API virtual bool CanConstructPropertyUpdatedWidgetBuilder() const override;
 
 	// Returns the builder used to construct the property updated widgets, in this case the SOverrideStatusWidget 
-	virtual TSharedPtr<FPropertyUpdatedWidgetBuilder> ConstructPropertyUpdatedWidgetBuilder(const FConstructPropertyUpdatedWidgetBuilderArgs& Args) override;
+	UE_API virtual TSharedPtr<FPropertyUpdatedWidgetBuilder> ConstructPropertyUpdatedWidgetBuilder(const FConstructPropertyUpdatedWidgetBuilderArgs& Args) override;
 
 	// Returns the preconfigured menu builder for this display manager and a given subject
-	TSharedPtr<FOverrideStatusWidgetMenuBuilder> GetMenuBuilder(const FOverrideStatusSubject& InSubject) const;
+	UE_API TSharedPtr<FOverrideStatusWidgetMenuBuilder> GetMenuBuilder(const FOverrideStatusSubject& InSubject) const;
 
 	FOverrideStatus_CanCreateWidget& OnCanCreateWidget() { return CanCreateWidgetDelegate; }
 	FOverrideStatus_GetStatus& OnGetStatus() { return GetStatusDelegate; }
@@ -51,9 +53,9 @@ public:
 
 private:
 	
-	TSharedPtr<FOverrideStatusDetailsWidgetBuilder> ConstructOverrideWidgetBuilder(const FConstructPropertyUpdatedWidgetBuilderArgs& Args);
-	void OnObjectPropertyChanged(UObject* ObjectBeingModified, FPropertyChangedEvent& PropertyChangedEvent);
-	void SetIsDisplayingOverrideableObject(bool InIsDisplayingOverrideableObject);
+	UE_API TSharedPtr<FOverrideStatusDetailsWidgetBuilder> ConstructOverrideWidgetBuilder(const FConstructPropertyUpdatedWidgetBuilderArgs& Args);
+	UE_API void OnObjectPropertyChanged(UObject* ObjectBeingModified, FPropertyChangedEvent& PropertyChangedEvent);
+	UE_API void SetIsDisplayingOverrideableObject(bool InIsDisplayingOverrideableObject);
 
 	bool bIsDisplayingOverrideableObject = false;
 	
@@ -70,3 +72,5 @@ private:
 
 	friend class FOverrideStatusDetailsViewObjectFilter;
 };
+
+#undef UE_API
diff --git a/Engine/Source/Editor/AnimationEditorWidgets/Public/Overrides/OverrideStatusDetailsObjectFilter.h b/Engine/Source/Editor/AnimationEditorWidgets/Public/Overrides/OverrideStatusDetailsObjectFilter.h
--- a/Engine/Source/Editor/AnimationEditorWidgets/Public/Overrides/OverrideStatusDetailsObjectFilter.h
+++ b/Engine/Source/Editor/AnimationEditorWidgets/Public/Overrides/OverrideStatusDetailsObjectFilter.h
@@ -6,6 +6,8 @@
 #include "Overrides/OverrideStatusWidgetMenuBuilder.h"
 #include "DetailsViewObjectFilter.h"
 
+#define UE_API ANIMATIONEDITORWIDGETS_API
+
 DECLARE_DELEGATE_RetVal_TwoParams(bool, FOverrideStatusObjectFilter_CanMergeObjects, const UObject*, const UObject*);
 
 /**
@@ -14,7 +16,7 @@ DECLARE_DELEGATE_RetVal_TwoParams(bool, FOverrideStatusObjectFilter_CanMergeObje
  * is used to provide the override status widget instead of the reset value arrow to the details
  * panel.
  */
-class ANIMATIONEDITORWIDGETS_API FOverrideStatusDetailsViewObjectFilter : public FDetailsViewObjectFilter
+class FOverrideStatusDetailsViewObjectFilter : public FDetailsViewObjectFilter
 {	
 public:
 	// the standard method to create an object filter. 
@@ -27,30 +29,30 @@ public:
 	}
 
 	// default constructor
-	FOverrideStatusDetailsViewObjectFilter();
+	UE_API FOverrideStatusDetailsViewObjectFilter();
 
 	// sets up the  display manager for this filter
-	virtual void InitializeDisplayManager();
+	UE_API virtual void InitializeDisplayManager();
 
 	// Given a const TArray<UObject*>& SourceObjects, filters the objects and puts the objects which should
 	// be shown in the Details panel in the return TArray<FDetailsViewObjectRoot>. These may be some part of the
 	// original SourceObjects array, itself, or it may be some contained sub-objects within the SourceObjects.
-	virtual TArray<FDetailsViewObjectRoot> FilterObjects(const TArray<UObject*>& SourceObjects) override;
+	UE_API virtual TArray<FDetailsViewObjectRoot> FilterObjects(const TArray<UObject*>& SourceObjects) override;
 
 	// Returns a preconfigured menu builder for this filter.
-	TSharedPtr<FOverrideStatusWidgetMenuBuilder> GetMenuBuilder(const FOverrideStatusSubject& InSubject) const;
+	UE_API TSharedPtr<FOverrideStatusWidgetMenuBuilder> GetMenuBuilder(const FOverrideStatusSubject& InSubject) const;
 	
-	FOverrideStatus_CanCreateWidget& OnCanCreateWidget();
-	FOverrideStatus_GetStatus& OnGetStatus();
-	FOverrideStatus_OnWidgetClicked& OnWidgetClicked();
-	FOverrideStatus_OnGetMenuContent& OnGetMenuContent();
-	FOverrideStatus_AddOverride& OnAddOverride();
-	FOverrideStatus_ClearOverride& OnClearOverride();
-	FOverrideStatus_ResetToDefault& OnResetToDefault();
-	FOverrideStatus_ValueDiffersFromDefault& OnValueDiffersFromDefault();
+	UE_API FOverrideStatus_CanCreateWidget& OnCanCreateWidget();
+	UE_API FOverrideStatus_GetStatus& OnGetStatus();
+	UE_API FOverrideStatus_OnWidgetClicked& OnWidgetClicked();
+	UE_API FOverrideStatus_OnGetMenuContent& OnGetMenuContent();
+	UE_API FOverrideStatus_AddOverride& OnAddOverride();
+	UE_API FOverrideStatus_ClearOverride& OnClearOverride();
+	UE_API FOverrideStatus_ResetToDefault& OnResetToDefault();
+	UE_API FOverrideStatus_ValueDiffersFromDefault& OnValueDiffersFromDefault();
 	FOverrideStatusObjectFilter_CanMergeObjects& OnCanMergeObjects() { return CanMergeObjectDelegate; }
 
-	static bool MergeObjectByClass(const UObject* InObjectA, const UObject* InObjectB);
+	static UE_API bool MergeObjectByClass(const UObject* InObjectA, const UObject* InObjectB);
 
 private:
 	
@@ -61,4 +63,6 @@ private:
 	TSharedPtr<FOverrideStatusDetailsDisplayManager> OverrideStatusDisplayManager;
 
 	FOverrideStatusObjectFilter_CanMergeObjects CanMergeObjectDelegate;
-};
\ No newline at end of file
+};
+
+#undef UE_API
diff --git a/Engine/Source/Editor/AnimationEditorWidgets/Public/Overrides/OverrideStatusDetailsWidgetBuilder.h b/Engine/Source/Editor/AnimationEditorWidgets/Public/Overrides/OverrideStatusDetailsWidgetBuilder.h
--- a/Engine/Source/Editor/AnimationEditorWidgets/Public/Overrides/OverrideStatusDetailsWidgetBuilder.h
+++ b/Engine/Source/Editor/AnimationEditorWidgets/Public/Overrides/OverrideStatusDetailsWidgetBuilder.h
@@ -6,12 +6,14 @@
 #include "UserInterface/Widgets/PropertyUpdatedWidgetBuilder.h"
 #include "OverrideStatusSubject.h"
 
+#define UE_API ANIMATIONEDITORWIDGETS_API
+
 class FOverrideStatusDetailsDisplayManager;
 
 /**
  * A Display builder for the override status combo button
  */
-class ANIMATIONEDITORWIDGETS_API FOverrideStatusDetailsWidgetBuilder : public FPropertyUpdatedWidgetBuilder
+class FOverrideStatusDetailsWidgetBuilder : public FPropertyUpdatedWidgetBuilder
 {
 public:
 	DECLARE_DELEGATE_RetVal(TSharedRef<SWidget>, FOnGetWidget);
@@ -25,7 +27,7 @@ public:
 	 * @param InDetailsDisplayManager the FDetailsDisplayManager which manages the details display
 	 * than a property row   
 	 */
-	FOverrideStatusDetailsWidgetBuilder(
+	UE_API FOverrideStatusDetailsWidgetBuilder(
 		const TSharedRef<FOverrideStatusDetailsDisplayManager>& InDetailsDisplayManager,
 		const TArray<FOverrideStatusObject>& InObjects,
 		const TSharedPtr<const FPropertyPath>& InPropertyPath,
@@ -34,21 +36,21 @@ public:
 	/**
 	 * Implements the generation of the Category Menu button SWidget
 	 */
-	virtual TSharedPtr<SWidget> GenerateWidget() override;
+	UE_API virtual TSharedPtr<SWidget> GenerateWidget() override;
 
 	/**
 	 * Converts this into the SWidget it builds
 	 */
-	virtual ~FOverrideStatusDetailsWidgetBuilder() override;
+	UE_API virtual ~FOverrideStatusDetailsWidgetBuilder() override;
 
-	FOverrideStatus_CanCreateWidget& OnCanCreateWidget();
-	FOverrideStatus_GetStatus& OnGetStatus();
-	FOverrideStatus_OnWidgetClicked& OnWidgetClicked();
-	FOverrideStatus_OnGetMenuContent& OnGetMenuContent();
-	FOverrideStatus_AddOverride& OnAddOverride();
-	FOverrideStatus_ClearOverride& OnClearOverride();
-	FOverrideStatus_ResetToDefault& OnResetToDefault();
-	FOverrideStatus_ValueDiffersFromDefault& OnValueDiffersFromDefault();
+	UE_API FOverrideStatus_CanCreateWidget& OnCanCreateWidget();
+	UE_API FOverrideStatus_GetStatus& OnGetStatus();
+	UE_API FOverrideStatus_OnWidgetClicked& OnWidgetClicked();
+	UE_API FOverrideStatus_OnGetMenuContent& OnGetMenuContent();
+	UE_API FOverrideStatus_AddOverride& OnAddOverride();
+	UE_API FOverrideStatus_ClearOverride& OnClearOverride();
+	UE_API FOverrideStatus_ResetToDefault& OnResetToDefault();
+	UE_API FOverrideStatus_ValueDiffersFromDefault& OnValueDiffersFromDefault();
 
 private:
 
@@ -59,4 +61,6 @@ private:
 	TSharedRef<FOverrideStatusDetailsDisplayManager> DisplayManager;
 
 	FOverrideStatusSubject Subject;
-};
\ No newline at end of file
+};
+
+#undef UE_API
diff --git a/Engine/Source/Editor/AnimationEditorWidgets/Public/Overrides/OverrideStatusWidgetMenuBuilder.h b/Engine/Source/Editor/AnimationEditorWidgets/Public/Overrides/OverrideStatusWidgetMenuBuilder.h
--- a/Engine/Source/Editor/AnimationEditorWidgets/Public/Overrides/OverrideStatusWidgetMenuBuilder.h
+++ b/Engine/Source/Editor/AnimationEditorWidgets/Public/Overrides/OverrideStatusWidgetMenuBuilder.h
@@ -8,6 +8,8 @@
 #include "Overrides/SOverrideStatusWidget.h"
 #include "Overrides/OverrideStatusDetailsWidgetBuilder.h"
 
+#define UE_API ANIMATIONEDITORWIDGETS_API
+
 class FPropertyPath;
 class FOverrideStatusDetailsDisplayManager;
 
@@ -15,7 +17,7 @@ class FOverrideStatusDetailsDisplayManager;
  * Builder for the override status menu that goes on the top right side of each property
  * Component in the details panel
  */
-class ANIMATIONEDITORWIDGETS_API FOverrideStatusWidgetMenuBuilder : public FToolElementRegistrationArgs
+class FOverrideStatusWidgetMenuBuilder : public FToolElementRegistrationArgs
 {
 
 	DECLARE_DELEGATE(FResetToDefault);
@@ -25,12 +27,12 @@ public:
 	 * constructor
 	 *
 	 */
-	FOverrideStatusWidgetMenuBuilder(const FOverrideStatusSubject& InSubject, const TWeakPtr<FOverrideStatusDetailsDisplayManager>& InDisplayManager = TWeakPtr<FOverrideStatusDetailsDisplayManager>());
+	UE_API FOverrideStatusWidgetMenuBuilder(const FOverrideStatusSubject& InSubject, const TWeakPtr<FOverrideStatusDetailsDisplayManager>& InDisplayManager = TWeakPtr<FOverrideStatusDetailsDisplayManager>());
 
-	virtual ~FOverrideStatusWidgetMenuBuilder() override;
+	UE_API virtual ~FOverrideStatusWidgetMenuBuilder() override;
 
 	// returns the status of the override
-	EOverrideWidgetStatus::Type GetStatus() const;
+	UE_API EOverrideWidgetStatus::Type GetStatus() const;
 
 	// returns the attribute backing up the status of the override
 	TAttribute<EOverrideWidgetStatus::Type>& GetStatusAttribute() { return StatusAttribute; }
@@ -38,34 +40,34 @@ public:
 	/**
 	 * Override the active overrideable object at the given property path
 	 */
-	void AddOverride();
-	bool CanAddOverride() const;
-	FOverrideStatus_AddOverride& OnAddOverride();
+	UE_API void AddOverride();
+	UE_API bool CanAddOverride() const;
+	UE_API FOverrideStatus_AddOverride& OnAddOverride();
 	
 	/**
 	 * Clears any active overrides on the property / object
 	 */
-	void ClearOverride();
-	bool CanClearOverride() const;
-	FOverrideStatus_ClearOverride& OnClearOverride();
+	UE_API void ClearOverride();
+	UE_API bool CanClearOverride() const;
+	UE_API FOverrideStatus_ClearOverride& OnClearOverride();
 
 	/**
 	 * Override the active overrideable object at the given property path
 	 */
-	void ResetToDefault();
-	bool CanResetToDefault() const;
-	FOverrideStatus_ResetToDefault& OnResetToDefault();
-	FOverrideStatus_ValueDiffersFromDefault& OnValueDiffersFromDefault();
+	UE_API void ResetToDefault();
+	UE_API bool CanResetToDefault() const;
+	UE_API FOverrideStatus_ResetToDefault& OnResetToDefault();
+	UE_API FOverrideStatus_ValueDiffersFromDefault& OnValueDiffersFromDefault();
 
 	/**
 	 * Set up the menu
 	 */
-	void InitializeMenu();
+	UE_API void InitializeMenu();
 
 	/**
 	 * Fill this method in with your Slate to create it
 	 */
-	virtual TSharedPtr<SWidget> GenerateWidget() override;
+	UE_API virtual TSharedPtr<SWidget> GenerateWidget() override;
 
 private:
 	/**
@@ -89,3 +91,4 @@ private:
 	FOverrideStatus_ValueDiffersFromDefault ValueDiffersFromDefaultDelegate;
 };
 
+#undef UE_API
diff --git a/Engine/Source/Editor/AnimationEditorWidgets/Public/Overrides/SOverrideListWidget.h b/Engine/Source/Editor/AnimationEditorWidgets/Public/Overrides/SOverrideListWidget.h
--- a/Engine/Source/Editor/AnimationEditorWidgets/Public/Overrides/SOverrideListWidget.h
+++ b/Engine/Source/Editor/AnimationEditorWidgets/Public/Overrides/SOverrideListWidget.h
@@ -5,7 +5,9 @@
 #include "Overrides/OverrideStatusSubject.h"
 #include "Widgets/Text/STextBlock.h"
 
-class ANIMATIONEDITORWIDGETS_API SOverrideListWidget : public SCompoundWidget
+#define UE_API ANIMATIONEDITORWIDGETS_API
+
+class SOverrideListWidget : public SCompoundWidget
 {
 	SLATE_DECLARE_WIDGET(SOverrideListWidget, SCompoundWidget)
 public:
@@ -18,12 +20,12 @@ public:
         SLATE_EVENT(FOverrideStatus_ClearOverride, OnClearOverride)
 	SLATE_END_ARGS()
 public:
-	SOverrideListWidget();
-	virtual ~SOverrideListWidget() override;
+	UE_API SOverrideListWidget();
+	UE_API virtual ~SOverrideListWidget() override;
 
-	void Construct(const FArguments& InArgs);
+	UE_API void Construct(const FArguments& InArgs);
 
-	virtual void Tick(const FGeometry& AllottedGeometry, const double InCurrentTime, const float InDeltaTime) override;
+	UE_API virtual void Tick(const FGeometry& AllottedGeometry, const double InCurrentTime, const float InDeltaTime) override;
 	
 private:
 
@@ -35,3 +37,5 @@ private:
 	TSharedPtr<STextBlock> TextBlock;
 	TOptional<uint32> LastHash;
 };
+
+#undef UE_API
diff --git a/Engine/Source/Editor/AnimationEditorWidgets/Public/Overrides/SOverrideStatusWidget.h b/Engine/Source/Editor/AnimationEditorWidgets/Public/Overrides/SOverrideStatusWidget.h
--- a/Engine/Source/Editor/AnimationEditorWidgets/Public/Overrides/SOverrideStatusWidget.h
+++ b/Engine/Source/Editor/AnimationEditorWidgets/Public/Overrides/SOverrideStatusWidget.h
@@ -6,6 +6,8 @@
 #include "Styling/StyleColors.h"
 #include "UObject/OverriddenPropertySet.h"
 
+#define UE_API ANIMATIONEDITORWIDGETS_API
+
 class SImage;
 
 namespace EOverrideWidgetStatus
@@ -23,7 +25,7 @@ namespace EOverrideWidgetStatus
 	};
 }
 
-struct ANIMATIONEDITORWIDGETS_API FOverrideStatusWidgetStyle
+struct FOverrideStatusWidgetStyle
 {
 	const FSlateBrush* Icon = nullptr;
 	const FSlateBrush* HoveredIcon = nullptr;
@@ -36,7 +38,7 @@ struct ANIMATIONEDITORWIDGETS_API FOverrideStatusWidgetStyle
 DECLARE_DELEGATE_RetVal_OneParam(const FOverrideStatusWidgetStyle*, FOverrideStatusWidgetGetStyle, EOverrideWidgetStatus::Type);
 DECLARE_DELEGATE_RetVal_OneParam(FText, FOverrideStatusWidgetGetTooltip, EOverrideWidgetStatus::Type);
 
-class ANIMATIONEDITORWIDGETS_API SOverrideStatusWidget : public SCompoundWidget
+class SOverrideStatusWidget : public SCompoundWidget
 {
 	SLATE_DECLARE_WIDGET(SOverrideStatusWidget, SCompoundWidget)
 public:
@@ -53,29 +55,29 @@ public:
 		SLATE_ATTRIBUTE(bool, IsHovered)
 	SLATE_END_ARGS()
 public:
-	SOverrideStatusWidget();
-	virtual ~SOverrideStatusWidget() override;
+	UE_API SOverrideStatusWidget();
+	UE_API virtual ~SOverrideStatusWidget() override;
 
-	void Construct(const FArguments& InArgs);
+	UE_API void Construct(const FArguments& InArgs);
 
-	EOverrideWidgetStatus::Type GetStatus() const;
-	void SetStatus(EOverrideWidgetStatus::Type InStatus);
+	UE_API EOverrideWidgetStatus::Type GetStatus() const;
+	UE_API void SetStatus(EOverrideWidgetStatus::Type InStatus);
 
-	virtual FReply OnMouseButtonDown(const FGeometry& MyGeometry, const FPointerEvent& MouseEvent) override;
-	virtual FReply OnMouseButtonUp(const FGeometry& MyGeometry, const FPointerEvent& MouseEvent) override;
-	virtual void OnMouseEnter(const FGeometry& MyGeometry, const FPointerEvent& MouseEvent) override;
-	virtual void OnMouseLeave(const FPointerEvent& MouseEvent) override;
+	UE_API virtual FReply OnMouseButtonDown(const FGeometry& MyGeometry, const FPointerEvent& MouseEvent) override;
+	UE_API virtual FReply OnMouseButtonUp(const FGeometry& MyGeometry, const FPointerEvent& MouseEvent) override;
+	UE_API virtual void OnMouseEnter(const FGeometry& MyGeometry, const FPointerEvent& MouseEvent) override;
+	UE_API virtual void OnMouseLeave(const FPointerEvent& MouseEvent) override;
 	
-	void QueueRepaint();
+	UE_API void QueueRepaint();
 
-	virtual void Tick(const FGeometry& AllottedGeometry, const double InCurrentTime, const float InDeltaTime) override;
+	UE_API virtual void Tick(const FGeometry& AllottedGeometry, const double InCurrentTime, const float InDeltaTime) override;
 
-	static const FOverrideStatusWidgetStyle* GetStyleFromStatus(EOverrideWidgetStatus::Type InStatus);
+	static UE_API const FOverrideStatusWidgetStyle* GetStyleFromStatus(EOverrideWidgetStatus::Type InStatus);
 
 private:
 	
-	FReply HandleClick();
-	static const FOverrideStatusWidgetStyle* GetStyleFromStatus_Fallback(EOverrideWidgetStatus::Type InStatus);
+	UE_API FReply HandleClick();
+	static UE_API const FOverrideStatusWidgetStyle* GetStyleFromStatus_Fallback(EOverrideWidgetStatus::Type InStatus);
 	
 	/**
 	 * Style used for the button.
@@ -94,3 +96,5 @@ private:
 	FOverrideStatusWidgetGetStyle OnGetStyle;
 	FOverrideStatusWidgetGetTooltip OnGetTooltip;
 };
+
+#undef UE_API
