Added new version of AddAnimationNotifyEvent() that supports copying/specifying relevant data from FAnimNotifyEvent such as TriggerWeightThreshold for example.

#rnx

diff --git a/Engine/Source/Editor/AnimationBlueprintLibrary/Private/AnimationBlueprintLibrary.cpp b/Engine/Source/Editor/AnimationBlueprintLibrary/Private/AnimationBlueprintLibrary.cpp
--- a/Engine/Source/Editor/AnimationBlueprintLibrary/Private/AnimationBlueprintLibrary.cpp
+++ b/Engine/Source/Editor/AnimationBlueprintLibrary/Private/AnimationBlueprintLibrary.cpp
@@ -61,6 +61,7 @@
 #include "UObject/UnrealNames.h"
 #include "Animation/AnimData/IAnimationDataModel.h"
 #include "AnimPose.h"
+#include "UObject/UnrealTypePrivate.h"
 
 #include UE_INLINE_GENERATED_CPP_BY_NAME(AnimationBlueprintLibrary)
 
@@ -710,6 +711,94 @@ UAnimNotify* UAnimationBlueprintLibrary::AddAnimationNotifyEvent(UAnimSequenceBa
 	return Notify;
 }
 
+UAnimNotify* UAnimationBlueprintLibrary::AddAnimationNotifyEventFromSource(UAnimSequenceBase* AnimationSequenceBase, float StartTime, FAnimNotifyEvent NotifyEvent, FName NotifyTrackName)
+{
+	UAnimNotify* ResNotify = nullptr;
+	
+	if (AnimationSequenceBase)
+	{
+		TObjectPtr<UAnimNotify> Notify = NotifyEvent.Notify;
+		const bool bValidNotify = Notify != nullptr;
+		const bool bValidNotifyState = NotifyEvent.NotifyStateClass != nullptr;
+		const bool bValidOuter = bValidNotify && Notify->GetOuter() == AnimationSequenceBase;
+		const bool bIsValidTrackName = IsValidAnimNotifyTrackName(AnimationSequenceBase, NotifyTrackName);
+		const bool bIsValidTime = IsValidTimeInternal(AnimationSequenceBase, StartTime);
+		
+		if (bValidNotify && bIsValidTrackName && bIsValidTime && !bValidNotifyState)
+		{
+			// Our actual new event we will copy.
+			FAnimNotifyEvent& NewEvent = AnimationSequenceBase->Notifies.AddDefaulted_GetRef();
+
+			// Copy over only blueprint exposed types (to support for adding info such as trigger weight threshold, filtering type, etc).
+			const UScriptStruct* ScriptStruct = FAnimNotifyEvent::StaticStruct();
+			for (TFieldIterator<FProperty> PropIt(ScriptStruct); PropIt; ++PropIt)
+			{
+				FProperty* Property = *PropIt;
+				if (Property->HasAnyPropertyFlags(CPF_BlueprintVisible | CPF_Edit))
+				{
+					Property->CopyCompleteValue_InContainer(&NewEvent, &NotifyEvent);
+				}
+			}
+
+			// Duplicate notify since its outer is not the target animation sequence.
+			if (!bValidOuter)
+			{
+				UAnimNotify* DuplicatedNotify = static_cast<UAnimNotify*>(StaticDuplicateObject(Notify, AnimationSequenceBase, NAME_None, RF_AllFlags, nullptr, EDuplicateMode::Normal, ~EInternalObjectFlags::RootSet));
+				check(DuplicatedNotify && DuplicatedNotify->CanBePlaced(AnimationSequenceBase));
+				
+				NewEvent.Notify = DuplicatedNotify;
+			}
+			else
+			{
+				NewEvent.Notify = Notify;
+			}
+
+			// Now we have our final notify instance.
+			ResNotify = NewEvent.Notify;
+			
+			// Update info on new notify.
+			NewEvent.NotifyName = NAME_None;
+			NewEvent.Link(AnimationSequenceBase, StartTime);
+			NewEvent.TriggerTimeOffset = GetTriggerTimeOffsetForType(AnimationSequenceBase->CalculateOffsetForNotify(StartTime));
+			NewEvent.TrackIndex = GetTrackIndexForAnimationNotifyTrackName(AnimationSequenceBase, NotifyTrackName);
+			NewEvent.NotifyStateClass = nullptr;
+			NewEvent.Guid = FGuid::NewGuid();		
+
+			// Refresh all cached data
+			AnimationSequenceBase->MarkPackageDirty();
+			AnimationSequenceBase->RefreshCacheData();
+		}
+		else
+		{
+			if (!bValidNotify)
+			{
+				UE_LOG(LogAnimationBlueprintLibrary, Warning, TEXT("Invalid Animation Notify in AddAnimationNotifyEventObject"));
+			}
+			
+			if (bValidNotifyState)
+			{
+				UE_LOG(LogAnimationBlueprintLibrary, Warning, TEXT("Adding Animation Notify States is not supported yet."))
+			}
+			
+			if (!bIsValidTrackName)
+			{
+				UE_LOG(LogAnimationBlueprintLibrary, Warning, TEXT("Animation Notify Track %s does not exist on Animation Sequence %s"), *NotifyTrackName.ToString(), *AnimationSequenceBase->GetName());
+			}
+
+			if (!bIsValidTime)
+			{
+				UE_LOG(LogAnimationBlueprintLibrary, Warning, TEXT("%f is outside of Animation Sequence %s range"), StartTime, *AnimationSequenceBase->GetName());
+			}
+		}
+	}
+	else
+	{
+		UE_LOG(LogAnimationBlueprintLibrary, Warning, TEXT("Invalid Animation Sequence for AddAnimationNotifyEventObject"));
+	}
+
+	return ResNotify;
+}
+
 UAnimNotifyState* UAnimationBlueprintLibrary::AddAnimationNotifyStateEvent(UAnimSequenceBase* AnimationSequenceBase, FName NotifyTrackName, float StartTime, float Duration, TSubclassOf<UAnimNotifyState> NotifyStateClass)
 {
 	UAnimNotifyState* NotifyState = nullptr;
diff --git a/Engine/Source/Editor/AnimationBlueprintLibrary/Public/AnimationBlueprintLibrary.h b/Engine/Source/Editor/AnimationBlueprintLibrary/Public/AnimationBlueprintLibrary.h
--- a/Engine/Source/Editor/AnimationBlueprintLibrary/Public/AnimationBlueprintLibrary.h
+++ b/Engine/Source/Editor/AnimationBlueprintLibrary/Public/AnimationBlueprintLibrary.h
@@ -241,16 +241,33 @@ public:
 	/** Adds an Animation Notify Event to Notify track in the given Animation with the given Notify creation data */
 	UFUNCTION(BlueprintCallable, Category = "AnimationBlueprintLibrary|NotifyEvents")
 	static UE_API UAnimNotify* AddAnimationNotifyEvent(UAnimSequenceBase* AnimationSequenceBase, FName NotifyTrackName, float StartTime, TSubclassOf<UAnimNotify> NotifyClass);
-
+	
+	UE_EXPERIMENTAL(5.7, "Adding a notify to an animation sequence from a FAnimNotifyEvent is expertimental.")
+	/**
+	 * Experimental: Adds a custom animation notify event to the specified animation sequence at the given time and on the specified notify track.
+	 * 
+	 * A new notify event will be created using relevant data copied from the provided NotifyEvent. The UAnimNotify is extracted from
+	 * NotifyEvent and, if it is not already owned by the animation sequence (not part of its outer), it will be automatically duplicated.
+	 * Only properties exposed to Blueprints in the NotifyEvent (such as TriggerWeightThreshold and NotifyTriggerChance) will be copied.
+	 * 
+	 * @param AnimationSequenceBase  The animation sequence to which the new notify event will be added.
+	 * @param StartTime              The time (in seconds) within the sequence where the notify should be placed.
+	 * @param NotifyEvent            The source notify event used to initialize the new one. Its UAnimNotify will be extracted and validated.
+	 * @param NotifyTrackName        The name of the notify track to add the notify to.
+	 * @return                       The UAnimNotify used in the new notify event (original or duplicated).
+	 */
+	UFUNCTION(BlueprintCallable, Category = "AnimationBlueprintLibrary|NotifyEvents")
+	static UAnimNotify* AddAnimationNotifyEventFromSource(UAnimSequenceBase* AnimationSequenceBase, float StartTime, FAnimNotifyEvent NotifyEvent, FName NotifyTrackName);
+	
 	/** Adds an Animation Notify State Event to Notify track in the given Animation with the given Notify State creation data */
 	UFUNCTION(BlueprintCallable, Category = "AnimationBlueprintLibrary|NotifyEvents")
 	static UE_API UAnimNotifyState* AddAnimationNotifyStateEvent(UAnimSequenceBase* AnimationSequenceBase, FName NotifyTrackName, float StartTime, float Duration, TSubclassOf<UAnimNotifyState> NotifyStateClass);
 
-	/** Adds an the specific Animation Notify to the Animation Sequence (requires Notify's outer to be the Animation Sequence) */
+	/** Adds the specific Animation Notify to the Animation Sequence (requires Notify's outer to be the Animation Sequence) */
 	UFUNCTION(BlueprintCallable, Category = "AnimationBlueprintLibrary|NotifyEvents")
 	static UE_API void AddAnimationNotifyEventObject(UAnimSequenceBase* AnimationSequenceBase, float StartTime, UAnimNotify* Notify, FName NotifyTrackName);
-
-	/** Adds an the specific Animation Notify State to the Animation Sequence (requires Notify State's outer to be the Animation Sequence) */
+	
+	/** Adds the specific Animation Notify State to the Animation Sequence (requires Notify State's outer to be the Animation Sequence) */
 	UFUNCTION(BlueprintCallable, Category = "AnimationBlueprintLibrary|NotifyEvents")
 	static UE_API void AddAnimationNotifyStateEventObject(UAnimSequenceBase* AnimationSequenceBase, float StartTime, float Duration, UAnimNotifyState* NotifyState, FName NotifyTrackName);
 
