Expose CurveMetaData API on the skeleton to Blueprint and Python

diff --git a/Engine/Source/Editor/AnimationBlueprintLibrary/Private/AnimationBlueprintLibrary.cpp b/Engine/Source/Editor/AnimationBlueprintLibrary/Private/AnimationBlueprintLibrary.cpp
--- a/Engine/Source/Editor/AnimationBlueprintLibrary/Private/AnimationBlueprintLibrary.cpp
+++ b/Engine/Source/Editor/AnimationBlueprintLibrary/Private/AnimationBlueprintLibrary.cpp
@@ -1656,6 +1656,127 @@ bool UAnimationBlueprintLibrary::DoesBoneNameExistInternal(USkeleton* Skeleton,
 	return Skeleton->GetReferenceSkeleton().FindBoneIndex(BoneName) != INDEX_NONE;
 }
 
+int32 UAnimationBlueprintLibrary::GetNumCurveMetaData(const USkeleton* Skeleton)
+{
+	if (Skeleton)
+	{
+		return Skeleton->GetNumCurveMetaData();
+	}
+	else
+	{
+		UE_LOG(LogAnimationBlueprintLibrary, Warning, TEXT("Invalid Skeleton ptr for GetNumCurveMetaData"));
+	}
+
+	return 0;
+}
+
+bool UAnimationBlueprintLibrary::AddCurveMetaData(USkeleton* Skeleton, FName CurveName, bool bTransact /*= true*/)
+{
+	if (Skeleton)
+	{
+		return Skeleton->AddCurveMetaData(CurveName, bTransact);
+	}
+	else
+	{
+		UE_LOG(LogAnimationBlueprintLibrary, Warning, TEXT("Invalid Skeleton ptr for AddCurveMetaData"));
+	}
+
+	return false;
+}
+
+void UAnimationBlueprintLibrary::GetCurveMetaDataNames(const USkeleton* Skeleton, TArray<FName>& OutNames)
+{
+	if (Skeleton)
+	{
+		Skeleton->GetCurveMetaDataNames(OutNames);
+	}
+	else
+	{
+		UE_LOG(LogAnimationBlueprintLibrary, Warning, TEXT("Invalid Skeleton ptr for GetCurveMetaDataNames"));
+	}
+}
+
+bool UAnimationBlueprintLibrary::RenameCurveMetaData(USkeleton* Skeleton, FName OldName, FName NewName)
+{
+	if (Skeleton)
+	{
+		return Skeleton->RenameCurveMetaData(OldName, NewName);
+	}
+	else
+	{
+		UE_LOG(LogAnimationBlueprintLibrary, Warning, TEXT("Invalid Skeleton ptr for RenameCurveMetaData"));
+	}
+
+	return false;
+}
+
+bool UAnimationBlueprintLibrary::RemoveCurveMetaData(USkeleton* Skeleton, FName CurveName)
+{
+	if (Skeleton)
+	{
+		return Skeleton->RemoveCurveMetaData(CurveName);
+	}
+	else
+	{
+		UE_LOG(LogAnimationBlueprintLibrary, Warning, TEXT("Invalid Skeleton ptr for RemoveCurveMetaData"));
+	}
+
+	return false;
+}
+
+bool UAnimationBlueprintLibrary::GetCurveMetaDataMaterial(const USkeleton* Skeleton, FName CurveName)
+{
+	if (Skeleton)
+	{
+		return Skeleton->GetCurveMetaDataMaterial(CurveName);
+	}
+	else
+	{
+		UE_LOG(LogAnimationBlueprintLibrary, Warning, TEXT("Invalid Skeleton ptr for GetCurveMetaDataMaterial"));
+	}
+
+	return false;
+}
+
+bool UAnimationBlueprintLibrary::GetCurveMetaDataMorphTarget(const USkeleton* Skeleton, FName CurveName)
+{
+	if (Skeleton)
+	{
+		return Skeleton->GetCurveMetaDataMorphTarget(CurveName);
+	}
+	else
+	{
+		UE_LOG(LogAnimationBlueprintLibrary, Warning, TEXT("Invalid Skeleton ptr for GetCurveMetaDataMorphTarget"));
+	}
+
+	return false;
+}
+
+void UAnimationBlueprintLibrary::SetCurveMetaDataMaterial(USkeleton* Skeleton, FName CurveName, bool bOverrideMaterial)
+{
+	if (Skeleton)
+	{
+		Skeleton->SetCurveMetaDataMaterial(CurveName, bOverrideMaterial);
+	}
+	else
+	{
+		UE_LOG(LogAnimationBlueprintLibrary, Warning, TEXT("Invalid Skeleton ptr for SetCurveMetaDataMaterial"));
+	}
+}
+
+
+void UAnimationBlueprintLibrary::SetCurveMetaDataMorphTarget(USkeleton* Skeleton, FName CurveName, bool bOverrideMorphTarget)
+{
+	if (Skeleton)
+	{
+		Skeleton->SetCurveMetaDataMorphTarget(CurveName, bOverrideMorphTarget);
+	}
+	else
+	{
+		UE_LOG(LogAnimationBlueprintLibrary, Warning, TEXT("Invalid Skeleton ptr for SetCurveMetaDataMorphTarget"));
+	}
+}
+
 void UAnimationBlueprintLibrary::GetFloatKeys(UAnimSequenceBase* AnimationSequenceBase, FName CurveName, TArray<float>& Times, TArray<float>& Values)
 {
 	if (AnimationSequenceBase)
diff --git a/Engine/Source/Editor/AnimationBlueprintLibrary/Public/AnimationBlueprintLibrary.h b/Engine/Source/Editor/AnimationBlueprintLibrary/Public/AnimationBlueprintLibrary.h
--- a/Engine/Source/Editor/AnimationBlueprintLibrary/Public/AnimationBlueprintLibrary.h
+++ b/Engine/Source/Editor/AnimationBlueprintLibrary/Public/AnimationBlueprintLibrary.h
@@ -387,6 +387,76 @@ public:
 	UE_DEPRECATED(5.3, "This function is no longer used.")
 	static bool DoesBoneCurveNameExistInternal(USkeleton* Skeleton, FName BoneName) { return false; }
 
+	/** @return the number of curve metadata entries on the skeleton **/
+	UFUNCTION(BlueprintPure, meta = (ScriptMethod), Category = "AnimationBlueprintLibrary|Skeleton")
+	static UE_API int32 GetNumCurveMetaData(const USkeleton* Skeleton);
+
+	/**
+	 * Adds a curve metadata entry with the specified name on the skeleton
+	 * @param	InCurveName			The name of the curve to find
+	 * @param   bTransact           If true record a new transaction
+	 * @return true if an entry was added, false if an entry already existed
+	 */
+	UFUNCTION(BlueprintCallable, meta = (ScriptMethod), Category = "AnimationBlueprintLibrary|Skeleton")
+	static UE_API bool AddCurveMetaData(USkeleton* Skeleton, FName CurveName, bool bTransact = true);
+
+	/**
+	 * Get an array of all curve metadata names
+	 * @param	OutNames	The array to receive the metadata names 
+	 */
+	UFUNCTION(BlueprintPure, meta = (ScriptMethod), Category = "AnimationBlueprintLibrary|Skeleton")
+	static UE_API void GetCurveMetaDataNames(const USkeleton* Skeleton, TArray<FName>& OutNames);
+
+	/**
+	 * Renames a curve metadata entry. Metadata is preserved, but assigned to a different curve name.
+	 * @param OldName	The name of an existing curve entry
+	 * @param NewName	The name to change the entry to
+	 * @return			true if the rename was successful (the old name was found and the new name didnt collide with an
+	 *					existing entry)
+	*/
+	UFUNCTION(BlueprintCallable, meta = (ScriptMethod), Category = "AnimationBlueprintLibrary|Skeleton")
+	static UE_API bool RenameCurveMetaData(USkeleton* Skeleton, FName OldName, FName NewName);
+
+	/**
+	 * Removes a curve metadata entry for the specified name.
+	 * @param CurveName	The name of the curve to remove the metadata for
+	 * @return true if the entry was successfully removed (i.e. it existed)
+	 */
+	UFUNCTION(BlueprintCallable, meta = (ScriptMethod), Category = "AnimationBlueprintLibrary|Skeleton")
+	static UE_API bool RemoveCurveMetaData(USkeleton* Skeleton, FName CurveName);
+
+	/**
+	 * Gets the material flag for a curve's metadata
+	 * @param	CurveName	The name of the curve to check
+	 * @return true if the flag has been set
+	 */
+	UFUNCTION(BlueprintPure, meta = (ScriptMethod), Category = "AnimationBlueprintLibrary|Skeleton")
+	static UE_API bool GetCurveMetaDataMaterial(const USkeleton* Skeleton, FName CurveName);
+
+	/**
+	 * Gets the morph target flag for a curve's metadata
+	 * @param	CurveName	The name of the curve to check
+	 * @return true if the flag has been set
+	 */
+	UFUNCTION(BlueprintPure, meta = (ScriptMethod), Category = "AnimationBlueprintLibrary|Skeleton")
+	static UE_API bool GetCurveMetaDataMorphTarget(const USkeleton* Skeleton, FName CurveName);
+
+	/**
+	 * Set the material flag for a curve's metadata
+	 * @param	CurveName			The name of the curve to set
+	 * @param	bOverrideMaterial	Whether to set the material flag
+	 */
+	UFUNCTION(BlueprintCallable, meta = (ScriptMethod), Category = "AnimationBlueprintLibrary|Skeleton")
+	static UE_API void SetCurveMetaDataMaterial(USkeleton* Skeleton, FName CurveName, bool bOverrideMaterial);
+
+	/**
+	 * Set the morph target flag for a curve's metadata
+	 * @param	CurveName				The name of the curve to set
+	 * @param	bOverrideMorphTarget	Whether to set the morph target flag
+	 */
+	UFUNCTION(BlueprintCallable, meta = (ScriptMethod), Category = "AnimationBlueprintLibrary|Skeleton")
+	static UE_API void SetCurveMetaDataMorphTarget(USkeleton* Skeleton, FName CurveName, bool bOverrideMorphTarget);
+
 	/** Retrieves, a multiple of, Float Key(s) from the specified Animation Curve inside of the given Animation Sequence */
 	UFUNCTION(BlueprintCallable, Category = "AnimationBlueprintLibrary|Curves")
 	static UE_API void GetFloatKeys(UAnimSequenceBase* AnimationSequenceBase, FName CurveName, TArray<float>& Times, TArray<float>& Values);
diff --git a/Engine/Source/Runtime/Engine/Classes/Animation/Skeleton.h b/Engine/Source/Runtime/Engine/Classes/Animation/Skeleton.h
--- a/Engine/Source/Runtime/Engine/Classes/Animation/Skeleton.h
+++ b/Engine/Source/Runtime/Engine/Classes/Animation/Skeleton.h
@@ -448,6 +448,34 @@ public:
 	 * @return true if any of the entries were successfully removed (i.e. something changed)
 	 */
 	ENGINE_API bool RemoveCurveMetaData(TArrayView<FName> CurveNames);
+
+	/**
+	 * Gets the material flag for a curve's metadata
+	 * @param	CurveName	The name of the curve to check
+	 * @return true if the flag has been set
+	 */
+	ENGINE_API bool GetCurveMetaDataMaterial(FName CurveName) const;
+
+	/**
+	 * Gets the morph target flag for a curve's metadata
+	 * @param	CurveName	The name of the curve to check
+	 * @return true if the flag has been set
+	 */
+	ENGINE_API bool GetCurveMetaDataMorphTarget(FName CurveName) const;
+
+	/**
+	 * Set the material flag for a curve's metadata
+	 * @param	CurveName			The name of the curve to set
+	 * @param	bOverrideMaterial	Whether to set the material flag
+	 */
+	ENGINE_API void SetCurveMetaDataMaterial(FName CurveName, bool bOverrideMaterial);
+
+	/**
+	 * Set the morph target flag for a curve's metadata
+	 * @param	CurveName				The name of the curve to set
+	 * @param	bOverrideMorphTarget	Whether to set the morph target flag
+	 */
+	ENGINE_API void SetCurveMetaDataMorphTarget(FName CurveName, bool bOverrideMorphTarget);
 #endif
 
 	// this is called when you know both flags - called by post serialize and import
diff --git a/Engine/Source/Runtime/Engine/Private/Animation/AnimCurveMetaData.cpp b/Engine/Source/Runtime/Engine/Private/Animation/AnimCurveMetaData.cpp
--- a/Engine/Source/Runtime/Engine/Private/Animation/AnimCurveMetaData.cpp
+++ b/Engine/Source/Runtime/Engine/Private/Animation/AnimCurveMetaData.cpp
@@ -199,6 +199,28 @@ bool UAnimCurveMetaData::RemoveCurveMetaData(TArrayView<FName> CurveNames)
 }
 
 
+bool UAnimCurveMetaData::GetCurveMetaDataMaterial(FName CurveName) const
+{
+	const FCurveMetaData* MetaData = GetCurveMetaData(CurveName);
+	if (MetaData)
+	{
+		return MetaData->Type.bMaterial;
+	}
+
+	return false;
+}
+
+bool UAnimCurveMetaData::GetCurveMetaDataMorphTarget(FName CurveName) const
+{
+	const FCurveMetaData* MetaData = GetCurveMetaData(CurveName);
+	if (MetaData)
+	{
+		return MetaData->Type.bMorphtarget;
+	}
+
+	return false;
+}
+
 void UAnimCurveMetaData::SetCurveMetaDataMaterial(FName CurveName, bool bOverrideMaterial)
 {
 	FCurveMetaData* MetaData = GetCurveMetaData(CurveName);
diff --git a/Engine/Source/Runtime/Engine/Private/Animation/Skeleton.cpp b/Engine/Source/Runtime/Engine/Private/Animation/Skeleton.cpp
--- a/Engine/Source/Runtime/Engine/Private/Animation/Skeleton.cpp
+++ b/Engine/Source/Runtime/Engine/Private/Animation/Skeleton.cpp
@@ -1835,6 +1835,42 @@ bool USkeleton::RemoveCurveMetaData(TArrayView<FName> CurveNames)
 	return false;
 }
 
+bool USkeleton::GetCurveMetaDataMaterial(FName CurveName) const
+{
+	if (const UAnimCurveMetaData* AnimCurveMetaData = const_cast<USkeleton*>(this)->GetAssetUserData<UAnimCurveMetaData>())
+	{
+		return AnimCurveMetaData->GetCurveMetaDataMaterial(CurveName);
+	}
+
+	return false;
+}
+
+bool USkeleton::GetCurveMetaDataMorphTarget(FName CurveName) const
+{
+	if (const UAnimCurveMetaData* AnimCurveMetaData = const_cast<USkeleton*>(this)->GetAssetUserData<UAnimCurveMetaData>())
+	{
+		return AnimCurveMetaData->GetCurveMetaDataMorphTarget(CurveName);
+	}
+
+	return false;
+}
+
+void USkeleton::SetCurveMetaDataMaterial(FName CurveName, bool bOverrideMaterial)
+{
+	if (UAnimCurveMetaData* AnimCurveMetaData = GetAssetUserData<UAnimCurveMetaData>())
+	{
+		AnimCurveMetaData->SetCurveMetaDataMaterial(CurveName, bOverrideMaterial);
+	}
+}
+
+void USkeleton::SetCurveMetaDataMorphTarget(FName CurveName, bool bOverrideMorphTarget)
+{
+	if (UAnimCurveMetaData* AnimCurveMetaData = GetAssetUserData<UAnimCurveMetaData>())
+	{
+		AnimCurveMetaData->SetCurveMetaDataMorphTarget(CurveName, bOverrideMorphTarget);
+	}
+}
+
 #endif // WITH_EDITOR
 
 uint16 USkeleton::GetAnimCurveUidVersion() const
diff --git a/Engine/Source/Runtime/Engine/Public/Animation/AnimCurveMetadata.h b/Engine/Source/Runtime/Engine/Public/Animation/AnimCurveMetadata.h
--- a/Engine/Source/Runtime/Engine/Public/Animation/AnimCurveMetadata.h
+++ b/Engine/Source/Runtime/Engine/Public/Animation/AnimCurveMetadata.h
@@ -147,6 +147,20 @@ public:
 	 */
 	virtual FDelegateHandle RegisterOnCurveMetaDataChanged(const FSimpleMulticastDelegate::FDelegate& InOnCurveMetaDataChanged) = 0;
 
+	/**
+	 * Gets the material flag for a curve's metadata
+	 * @param	CurveName	The name of the curve to check
+	 * @return	true if the flag has been set
+	 */
+	virtual bool GetCurveMetaDataMaterial(FName CurveName) const = 0;
+
+	/**
+	 * Gets the morph target flag for a curve's metadata
+	 * @param	CurveName	The name of the curve to check
+	 * @return	true if the flag has been set
+	 */
+	virtual bool GetCurveMetaDataMorphTarget(FName CurveName) const = 0;
+
 	/**
 	 * Set the material flag for a curve's metadata
 	 * @param	CurveName			The name of the curve to set
@@ -210,6 +224,8 @@ public:
 	virtual bool RenameCurveMetaData(FName OldName, FName NewName) override;
 	virtual bool RemoveCurveMetaData(FName CurveName) override;
 	virtual bool RemoveCurveMetaData(TArrayView<FName> CurveNames) override;
+	virtual bool GetCurveMetaDataMaterial(FName CurveName) const override;
+	virtual bool GetCurveMetaDataMorphTarget(FName CurveName) const override;
 	virtual void SetCurveMetaDataMaterial(FName CurveName, bool bOverrideMaterial) override;
 	virtual void SetCurveMetaDataMorphTarget(FName CurveName, bool bOverrideMorphTarget) override;
 	virtual void SetCurveMetaDataBoneLinks(FName CurveName, const TArray<FBoneReference>& BoneLinks, uint8 InMaxLOD, USkeleton* InSkeleton) override;
