[CrashReport][Assert] UnrealEditor-AnimationBlueprintLibrary!FAnimPose::SetPose(FAnimationPoseData const &) [Engine/Source/Editor/AnimationBlueprintLibrary/Private/AnimPose.cpp:166]
#fix additional bone data indexing checks

diff --git a/Engine/Source/Editor/AnimationBlueprintLibrary/Private/AnimPose.cpp b/Engine/Source/Editor/AnimationBlueprintLibrary/Private/AnimPose.cpp
--- a/Engine/Source/Editor/AnimationBlueprintLibrary/Private/AnimPose.cpp
+++ b/Engine/Source/Editor/AnimationBlueprintLibrary/Private/AnimPose.cpp
@@ -116,10 +116,24 @@ void FAnimPose::SetPose(USkeletalMeshComponent* Component)
 		for (const FBoneIndexType BoneIndex : ContextBoneContainer.GetBoneIndicesArray())
 		{
 			const int32 SkeletonBoneIndex = ContextBoneContainer.GetSkeletonIndex(FCompactPoseBoneIndex(BoneIndex));
-			LocalSpacePoses[BoneIndices.IndexOfByKey(SkeletonBoneIndex)] = BoneSpaceTransforms[BoneIndex];
+			if (SkeletonBoneIndex != INDEX_NONE)
+			{
+				const int32 LocalBoneIndex = BoneIndices.IndexOfByKey(SkeletonBoneIndex);
+				if (LocalBoneIndex != INDEX_NONE)
+				{
+					LocalSpacePoses[LocalBoneIndex] = BoneSpaceTransforms[BoneIndex];
+				}
+				else
+				{
+					UE_LOG(LogAnimationPoseScripting, Warning, TEXT("Failed to find local Bone Index for Skeleton Bone Index %i [%s]"), BoneIndex, *ContextBoneContainer.GetReferenceSkeleton().GetBoneName(SkeletonBoneIndex).ToString());
+				}
+			}
+			else
+			{
+				UE_LOG(LogAnimationPoseScripting, Warning, TEXT("Failed to find Skeleton index for Compact Pose Bone Index %i"), BoneIndex);
+			}
 		}
 
-		ensure(LocalSpacePoses.Num() == RefLocalSpacePoses.Num());	
 		GenerateWorldSpaceTransforms();
 	}
 }
@@ -165,7 +179,22 @@ void FAnimPose::SetPose(const FAnimationPoseData& PoseData)
 		for (const FCompactPoseBoneIndex BoneIndex : CompactPose.ForEachBoneIndex())
 		{
 			const int32 SkeletonBoneIndex = ContextBoneContainer.GetSkeletonIndex(BoneIndex);
-			LocalSpacePoses[BoneIndices.IndexOfByKey(SkeletonBoneIndex)] = CompactPose[BoneIndex];
+			if (SkeletonBoneIndex != INDEX_NONE)
+			{
+				const int32 LocalBoneIndex = BoneIndices.IndexOfByKey(SkeletonBoneIndex);
+				if (LocalBoneIndex != INDEX_NONE)
+				{
+					LocalSpacePoses[LocalBoneIndex] = CompactPose[BoneIndex];
+				}
+				else
+				{
+					UE_LOG(LogAnimationPoseScripting, Warning, TEXT("Failed to find local Bone Index for Skeleton Bone Index %i [%s]"), BoneIndex.GetInt(), *ContextBoneContainer.GetReferenceSkeleton().GetBoneName(SkeletonBoneIndex).ToString());
+				}
+			}
+			else
+			{
+				UE_LOG(LogAnimationPoseScripting, Warning, TEXT("Failed to find Skeleton index for Compact Pose Bone Index %i"), BoneIndex.GetInt());
+			}
 		}
 
 		ensure(LocalSpacePoses.Num() == RefLocalSpacePoses.Num());
