Fixup incorrect bone index values being logged from CL 46304658 (bug-hawk)

diff --git a/Engine/Source/Editor/AnimationBlueprintLibrary/Private/AnimPose.cpp b/Engine/Source/Editor/AnimationBlueprintLibrary/Private/AnimPose.cpp
--- a/Engine/Source/Editor/AnimationBlueprintLibrary/Private/AnimPose.cpp
+++ b/Engine/Source/Editor/AnimationBlueprintLibrary/Private/AnimPose.cpp
@@ -125,7 +125,7 @@ void FAnimPose::SetPose(USkeletalMeshComponent* Component)
 				}
 				else
 				{
-					UE_LOG(LogAnimationPoseScripting, Warning, TEXT("Failed to find local Bone Index for Skeleton Bone Index %i [%s]"), BoneIndex, *ContextBoneContainer.GetReferenceSkeleton().GetBoneName(SkeletonBoneIndex).ToString());
+					UE_LOG(LogAnimationPoseScripting, Warning, TEXT("Failed to find local Bone Index for Skeleton Bone Index %i [%s]"), SkeletonBoneIndex, *ContextBoneContainer.GetReferenceSkeleton().GetBoneName(SkeletonBoneIndex).ToString());
 				}
 			}
 			else
@@ -188,7 +188,7 @@ void FAnimPose::SetPose(const FAnimationPoseData& PoseData)
 				}
 				else
 				{
-					UE_LOG(LogAnimationPoseScripting, Warning, TEXT("Failed to find local Bone Index for Skeleton Bone Index %i [%s]"), BoneIndex.GetInt(), *ContextBoneContainer.GetReferenceSkeleton().GetBoneName(SkeletonBoneIndex).ToString());
+					UE_LOG(LogAnimationPoseScripting, Warning, TEXT("Failed to find local Bone Index for Skeleton Bone Index %i [%s]"), SkeletonBoneIndex, *ContextBoneContainer.GetReferenceSkeleton().GetBoneName(SkeletonBoneIndex).ToString());
 				}
 			}
 			else
