Add python bindings for set/get preview mesh and additional meshes in USkeleton
* Exposes get_skeleton_preview_mesh, set_skeleton_preview_mesh(), get_skeleton_additional_previewl_meshes(), and set_skeleton_additional_preview_meshes() on unreal.AnimationLibrary

#Synced-CL 50482698

diff --git a/Engine/Source/Editor/AnimationBlueprintLibrary/Private/AnimationBlueprintLibrary.cpp b/Engine/Source/Editor/AnimationBlueprintLibrary/Private/AnimationBlueprintLibrary.cpp
--- a/Engine/Source/Editor/AnimationBlueprintLibrary/Private/AnimationBlueprintLibrary.cpp
+++ b/Engine/Source/Editor/AnimationBlueprintLibrary/Private/AnimationBlueprintLibrary.cpp
@@ -2967,4 +2967,38 @@ template void UAnimationBlueprintLibrary::GetCurveKeysInternal<float, FFloatCurv
 template void UAnimationBlueprintLibrary::GetCurveKeysInternal<FVector, FVectorCurve, ERawCurveTrackTypes::RCT_Vector>(UAnimSequenceBase* AnimationSequenceBase, FName CurveName, TArray<float>& Times, TArray<FVector>& KeyData);
 template void UAnimationBlueprintLibrary::GetCurveKeysInternal<FTransform, FTransformCurve, ERawCurveTrackTypes::RCT_Transform>(UAnimSequenceBase* AnimationSequenceBase, FName CurveName, TArray<float>& Times, TArray<FTransform>& KeyData);
 
+void UAnimationBlueprintLibrary::SetSkeletonPreviewMesh(USkeleton* Skeleton, USkeletalMesh* PreviewMesh, bool bMarkAsDirty)
+{
+	if (Skeleton)
+	{
+		Skeleton->SetPreviewMesh(PreviewMesh, bMarkAsDirty);
+	}
+}
+
+ USkeletalMesh* UAnimationBlueprintLibrary::GetSkeletonPreviewMesh(USkeleton* Skeleton, bool bFindIfNotSet)
+ {
+	if (Skeleton)
+	{
+		return Skeleton->GetPreviewMesh(bFindIfNotSet);
+	}
+	return nullptr;
+ }
+
+ void UAnimationBlueprintLibrary::SetSkeletonAdditionalPreviewMeshes(USkeleton* Skeleton, UDataAsset* AdditionalMeshes)
+ {
+	if (Skeleton)
+	{
+		Skeleton->SetAdditionalPreviewSkeletalMeshes(AdditionalMeshes);
+	}
+ }
+
+ UDataAsset* UAnimationBlueprintLibrary::GetSkeletonAdditionalPreviewMeshes(const USkeleton* Skeleton)
+ {
+	if (Skeleton)
+	{
+		return Skeleton->GetAdditionalPreviewSkeletalMeshes();
+	}
+	return nullptr;
+ }
+
 #undef LOCTEXT_NAMESPACE // "AnimationBlueprintLibrary"
diff --git a/Engine/Source/Editor/AnimationBlueprintLibrary/Public/AnimationBlueprintLibrary.h b/Engine/Source/Editor/AnimationBlueprintLibrary/Public/AnimationBlueprintLibrary.h
--- a/Engine/Source/Editor/AnimationBlueprintLibrary/Public/AnimationBlueprintLibrary.h
+++ b/Engine/Source/Editor/AnimationBlueprintLibrary/Public/AnimationBlueprintLibrary.h
@@ -717,6 +717,22 @@ public:
 	 */
 	UFUNCTION(BlueprintCallable, Category=Animation, meta = (ScriptMethod))
 	static UE_API void AddNodeAssetOverride(UAnimBlueprint* AnimBlueprint, const UAnimationAsset* Target, UAnimationAsset* Override, bool bPrintAppliedOverrides = false);
+
+	/** Sets the preview skeletal mesh on a skeleton, optionally marks the skeleton as dirty */
+	UFUNCTION(BlueprintCallable, Category=Skeleton, meta = (ScriptMethod))
+	static UE_API void SetSkeletonPreviewMesh(USkeleton* Skeleton, USkeletalMesh* PreviewMesh, bool bMarkAsDirty=true);
+
+	/** Returns the preview skeletal mesh on a skeleton, optionally finds one if a preview mesh hasn't been set */
+	UFUNCTION(BlueprintCallable, Category=Skeleton, meta = (ScriptMethod))
+	static UE_API USkeletalMesh* GetSkeletonPreviewMesh(USkeleton* Skeleton, bool bFindIfNotSet = false);
+
+	/** Sets the additional preview skeletal meshes on a skeleton */
+	UFUNCTION(BlueprintCallable, Category=Skeleton, meta = (ScriptMethod))
+	static UE_API void SetSkeletonAdditionalPreviewMeshes(USkeleton* Skeleton, UDataAsset* AdditionalMeshes);
+
+	/** Returns the additional preview skeletal meshes on a skeleton */
+	UFUNCTION(BlueprintPure, Category=Skeleton, meta = (ScriptMethod))
+	static UE_API UDataAsset* GetSkeletonAdditionalPreviewMeshes(const USkeleton* Skeleton);
 };
 
 #undef UE_API
